Control Rig Dynamics - add to experimental

#Synced-CL 52541170
#Synced-CL 52541170

#Synced-CL 52541170

#Synced-CL 52541170

diff --git a/Engine/Plugins/Experimental/ControlRigDynamics/ControlRigDynamics.uplugin b/Engine/Plugins/Experimental/ControlRigDynamics/ControlRigDynamics.uplugin
new file mode 100644
--- /dev/null
+++ b/Engine/Plugins/Experimental/ControlRigDynamics/ControlRigDynamics.uplugin
@@ -0,0 +1,35 @@
+{
+	"FileVersion": 3,
+	"Version": 1,
+	"VersionName": "1.0",
+	"FriendlyName": "ControlRigDynamics",
+	"Description": "Support for simple dynamics/cosmetic simulation in control rig",
+	"Category": "Animation",
+	"CreatedBy": "Epic Games, Inc.",
+	"CreatedByURL": "https://epicgames.com",
+	"DocsURL": "",
+	"MarketplaceURL": "",
+	"SupportURL": "",
+	"CanContainContent": true,
+	"CanContainVerse": false,
+	"IsBetaVersion": false,
+	"IsExperimentalVersion": true,
+	"Installed": false,
+	"Plugins": [
+		{
+			"Name": "ControlRig",
+			"Enabled": true
+		},
+		{
+			"Name": "PhysicsControl",
+			"Enabled": true
+		}
+	],
+	"Modules": [
+		{
+			"Name": "ControlRigDynamics",
+			"Type": "Runtime",
+			"LoadingPhase": "Default"
+		},
+	]
+}
\ No newline at end of file
diff --git a/Engine/Plugins/Experimental/ControlRigDynamics/Source/ControlRigDynamics/ControlRigDynamics.Build.cs b/Engine/Plugins/Experimental/ControlRigDynamics/Source/ControlRigDynamics/ControlRigDynamics.Build.cs
new file mode 100644
--- /dev/null
+++ b/Engine/Plugins/Experimental/ControlRigDynamics/Source/ControlRigDynamics/ControlRigDynamics.Build.cs
@@ -0,0 +1,44 @@
+// Copyright Epic Games, Inc. All Rights Reserved.
+
+using UnrealBuildTool;
+
+public class ControlRigDynamics : ModuleRules
+{
+	public ControlRigDynamics(ReadOnlyTargetRules Target) : base(Target)
+	{
+		PCHUsage = ModuleRules.PCHUsageMode.UseExplicitOrSharedPCHs;
+		//OptimizeCode = CodeOptimization.Never;
+
+		PublicDependencyModuleNames.AddRange(
+			new string[]
+			{
+				"Core",
+				"PhysicsControl"
+			}
+		);
+		
+		PrivateDependencyModuleNames.AddRange(
+			new string[]
+			{
+				"ControlRig",
+				"CoreUObject",
+				"Engine",
+				"RigVM"
+			}
+		);
+
+		if (Target.bBuildEditor == true)
+		{
+			PrivateDependencyModuleNames.AddRange(
+				new string[]
+				{
+					"RigVMDeveloper",
+					"ControlRigDeveloper",
+					"Slate",
+					"SlateCore",
+					"EditorStyle",
+				}
+			);
+		}
+	}
+}
diff --git a/Engine/Plugins/Experimental/ControlRigDynamics/Source/ControlRigDynamics/Private/ControlRigDynamicsModule.cpp b/Engine/Plugins/Experimental/ControlRigDynamics/Source/ControlRigDynamics/Private/ControlRigDynamicsModule.cpp
new file mode 100644
--- /dev/null
+++ b/Engine/Plugins/Experimental/ControlRigDynamics/Source/ControlRigDynamics/Private/ControlRigDynamicsModule.cpp
@@ -0,0 +1,32 @@
+// Copyright Epic Games, Inc. All Rights Reserved.
+
+#include "ControlRigDynamicsModule.h"
+#include "RigDynamicsObjectVersion.h"
+
+#include "UObject/DevObjectVersion.h"
+
+#define LOCTEXT_NAMESPACE "FControlRigDynamicsModule"
+
+DEFINE_LOG_CATEGORY(LogControlRigDynamics);
+
+
+// Unique version id for serializing
+const FGuid FRigDynamicsObjectVersion::GUID(0xA7C3E291, 0x5F8B4D06, 0x3E1A9C74, 0xD20F68B5);
+
+// Register custom version with Core
+FDevVersionRegistration GRegisterDynamicsObjectVersion(
+	FRigDynamicsObjectVersion::GUID, FRigDynamicsObjectVersion::LatestVersion, TEXT("Dev-RigDynamics"));
+
+//======================================================================================================================
+void FControlRigDynamicsModule::StartupModule()
+{
+}
+
+//======================================================================================================================
+void FControlRigDynamicsModule::ShutdownModule()
+{
+}
+
+#undef LOCTEXT_NAMESPACE
+
+IMPLEMENT_MODULE(FControlRigDynamicsModule, ControlRigDynamics)
diff --git a/Engine/Plugins/Experimental/ControlRigDynamics/Source/ControlRigDynamics/Private/RigDynamicsColliderComponent.cpp b/Engine/Plugins/Experimental/ControlRigDynamics/Source/ControlRigDynamics/Private/RigDynamicsColliderComponent.cpp
new file mode 100644
--- /dev/null
+++ b/Engine/Plugins/Experimental/ControlRigDynamics/Source/ControlRigDynamics/Private/RigDynamicsColliderComponent.cpp
@@ -0,0 +1,31 @@
+// Copyright Epic Games, Inc. All Rights Reserved.
+
+#include "RigDynamicsColliderComponent.h"
+
+#include "RigDynamicsObjectVersion.h"
+
+//======================================================================================================================
+void FRigDynamicsColliderComponent::Save(FArchive& Ar)
+{
+	Ar.UsingCustomVersion(FRigDynamicsObjectVersion::GUID);
+
+	FRigBaseComponent::Save(Ar);
+	Ar << Colliders;
+}
+
+//======================================================================================================================
+void FRigDynamicsColliderComponent::Load(FArchive& Ar)
+{
+	FRigBaseComponent::Load(Ar);
+
+	if (Ar.CustomVer(FRigDynamicsObjectVersion::GUID) >= FRigDynamicsObjectVersion::HelperStructs)
+	{
+		Ar << Colliders;
+	}
+	else
+	{
+		Ar << Colliders.Boxes;
+		Ar << Colliders.Capsules;
+		Ar << Colliders.Planes;
+	}
+}
diff --git a/Engine/Plugins/Experimental/ControlRigDynamics/Source/ControlRigDynamics/Private/RigDynamicsColliderExecution.cpp b/Engine/Plugins/Experimental/ControlRigDynamics/Source/ControlRigDynamics/Private/RigDynamicsColliderExecution.cpp
new file mode 100644
--- /dev/null
+++ b/Engine/Plugins/Experimental/ControlRigDynamics/Source/ControlRigDynamics/Private/RigDynamicsColliderExecution.cpp
@@ -0,0 +1,99 @@
+// Copyright Epic Games, Inc. All Rights Reserved.
+
+#include "RigDynamicsColliderExecution.h"
+
+#include "RigDynamicsHelpers.h"
+
+#include "Rigs/RigHierarchyController.h"
+#include "Units/RigUnitContext.h"
+#include "Units/Execution/RigUnit_PrepareForExecution.h"
+
+#include UE_INLINE_GENERATED_CPP_BY_NAME(RigDynamicsColliderExecution)
+
+//======================================================================================================================
+FRigUnit_SpawnDynamicsCollider_Execute()
+{
+	if (!ExecuteContext.Hierarchy)
+	{
+		return;
+	}
+	if (ExecuteContext.GetEventName() != FRigUnit_PrepareForExecution::EventName)
+	{
+		UE_CONTROLRIG_RIGUNIT_REPORT_ERROR(TEXT("SpawnDynamicsCollider can only be used during Setup"));
+		return;
+	}
+
+	if (URigHierarchyController* Controller = ExecuteContext.Hierarchy->GetController())
+	{
+		FRigHierarchyControllerInstructionBracket InstructionBracket(
+			Controller, ExecuteContext.GetInstructionIndex());
+		DynamicsColliderComponentKey = Controller->AddComponent(
+			FRigDynamicsColliderComponent::StaticStruct(), FRigDynamicsColliderComponent::GetDefaultName(), Owner);
+		if (DynamicsColliderComponentKey.IsValid())
+		{
+			if (FRigDynamicsColliderComponent* Component =
+				GetCollider(*ExecuteContext.Hierarchy, DynamicsColliderComponentKey))
+			{
+				Component->Colliders.Boxes = Colliders.Boxes;
+				Component->Colliders.Capsules = Colliders.Capsules;
+				Component->Colliders.Planes = Colliders.Planes;
+			}
+
+			if (FRigDynamicsSolverComponent* Solver = GetSolver(*ExecuteContext.Hierarchy, DynamicsSolverComponentKey))
+			{
+				Solver->Colliders.Add(DynamicsColliderComponentKey);
+			}
+		}
+	}
+}
+
+//======================================================================================================================
+// Macros to reduce boilerplate for the Set collider shape property nodes. Note that these can't
+// be used in the header because UHT scans raw source before C++ preprocessing, so it would
+// never see the USTRUCT/UPROPERTY/RIGVM_METHOD markers inside a macro expansion.
+//======================================================================================================================
+
+// Sets a whole shape by matching on Name (e.g. ShapeType=Box, ShapeArray=Boxes)
+#define IMPLEMENT_SET_DYNAMICS_COLLIDER_SHAPE(ShapeType, ShapeArray)           \
+FRigUnit_HierarchySetDynamicsCollider##ShapeType##_Execute()                   \
+{                                                                              \
+	if (!ExecuteContext.Hierarchy) { return; }                                 \
+	if (FRigDynamicsColliderComponent* Component =                            \
+		GetCollider(*ExecuteContext.Hierarchy, DynamicsColliderComponentKey))  \
+	{                                                                         \
+		for (auto& Existing : Component->Colliders.ShapeArray)                \
+		{                                                                     \
+			if (Existing.Name == ShapeType.Name) { Existing = ShapeType; }    \
+		}                                                                     \
+	}                                                                         \
+}
+
+// Sets an individual property on a shape found by Name
+#define IMPLEMENT_SET_DYNAMICS_COLLIDER_SHAPE_PROP(ShapeType, ShapeArray, PropName) \
+FRigUnit_HierarchySetDynamicsCollider##ShapeType##PropName##_Execute()              \
+{                                                                                   \
+	if (!ExecuteContext.Hierarchy) { return; }                                      \
+	if (FRigDynamicsColliderComponent* Component =                                 \
+		GetCollider(*ExecuteContext.Hierarchy, DynamicsColliderComponentKey))       \
+	{                                                                              \
+		for (auto& Shape : Component->Colliders.ShapeArray)                        \
+		{                                                                          \
+			if (Shape.Name == Name) { Shape.PropName = PropName; }                 \
+		}                                                                          \
+	}                                                                              \
+}
+
+//======================================================================================================================
+IMPLEMENT_SET_DYNAMICS_COLLIDER_SHAPE(Box, Boxes)
+IMPLEMENT_SET_DYNAMICS_COLLIDER_SHAPE(Capsule, Capsules)
+IMPLEMENT_SET_DYNAMICS_COLLIDER_SHAPE(Plane, Planes)
+//======================================================================================================================
+IMPLEMENT_SET_DYNAMICS_COLLIDER_SHAPE_PROP(Box, Boxes, TM)
+IMPLEMENT_SET_DYNAMICS_COLLIDER_SHAPE_PROP(Box, Boxes, Extents)
+//======================================================================================================================
+IMPLEMENT_SET_DYNAMICS_COLLIDER_SHAPE_PROP(Capsule, Capsules, TM)
+IMPLEMENT_SET_DYNAMICS_COLLIDER_SHAPE_PROP(Capsule, Capsules, Radius)
+IMPLEMENT_SET_DYNAMICS_COLLIDER_SHAPE_PROP(Capsule, Capsules, Length)
+//======================================================================================================================
+IMPLEMENT_SET_DYNAMICS_COLLIDER_SHAPE_PROP(Plane, Planes, TM)
+IMPLEMENT_SET_DYNAMICS_COLLIDER_SHAPE_PROP(Plane, Planes, Extents)
diff --git a/Engine/Plugins/Experimental/ControlRigDynamics/Source/ControlRigDynamics/Private/RigDynamicsConstraintComponent.cpp b/Engine/Plugins/Experimental/ControlRigDynamics/Source/ControlRigDynamics/Private/RigDynamicsConstraintComponent.cpp
new file mode 100644
--- /dev/null
+++ b/Engine/Plugins/Experimental/ControlRigDynamics/Source/ControlRigDynamics/Private/RigDynamicsConstraintComponent.cpp
@@ -0,0 +1,33 @@
+// Copyright Epic Games, Inc. All Rights Reserved.
+
+#include "RigDynamicsConstraintComponent.h"
+
+#include "RigDynamicsObjectVersion.h"
+
+//======================================================================================================================
+void FRigDynamicsConstraintComponent::Serialize(FArchive& Ar)
+{
+	Ar << ParentComponentKey;
+	Ar << ChildComponentKey;
+	Ar << ConstraintType;
+	Ar << Strength;
+	Ar << DampingRatio;
+	Ar << ExtraDamping;
+	Ar << LengthMultiplier;
+	Ar << ExtraLength;
+}
+
+//======================================================================================================================
+void FRigDynamicsConstraintComponent::Save(FArchive& Ar)
+{
+	Ar.UsingCustomVersion(FRigDynamicsObjectVersion::GUID);
+	FRigBaseComponent::Save(Ar);
+	Serialize(Ar);
+}
+
+//======================================================================================================================
+void FRigDynamicsConstraintComponent::Load(FArchive& Ar)
+{
+	FRigBaseComponent::Load(Ar);
+	Serialize(Ar);
+}
diff --git a/Engine/Plugins/Experimental/ControlRigDynamics/Source/ControlRigDynamics/Private/RigDynamicsConstraintExecution.cpp b/Engine/Plugins/Experimental/ControlRigDynamics/Source/ControlRigDynamics/Private/RigDynamicsConstraintExecution.cpp
new file mode 100644
--- /dev/null
+++ b/Engine/Plugins/Experimental/ControlRigDynamics/Source/ControlRigDynamics/Private/RigDynamicsConstraintExecution.cpp
@@ -0,0 +1,97 @@
+// Copyright Epic Games, Inc. All Rights Reserved.
+
+#include "RigDynamicsConstraintExecution.h"
+
+#include "RigDynamicsHelpers.h"
+
+#include "Rigs/RigHierarchyController.h"
+#include "Units/RigUnitContext.h"
+#include "Units/Execution/RigUnit_PrepareForExecution.h"
+
+#include UE_INLINE_GENERATED_CPP_BY_NAME(RigDynamicsConstraintExecution)
+
+//======================================================================================================================
+// Macros to reduce boilerplate for the Get/Set constraint property nodes. Note that these can't
+// be used in the header because UHT scans raw source before C++ preprocessing, so it would
+// never see the USTRUCT/UPROPERTY/RIGVM_METHOD markers inside a macro expansion.
+//======================================================================================================================
+
+#define IMPLEMENT_SET_DYNAMICS_CONSTRAINT(PropName)                               \
+FRigUnit_HierarchySetDynamicsConstraint##PropName##_Execute()                     \
+{                                                                                 \
+	if (!ExecuteContext.Hierarchy) { return; }                                    \
+	if (FRigDynamicsConstraintComponent* Component =                              \
+		GetConstraint(*ExecuteContext.Hierarchy, DynamicsConstraintComponentKey)) \
+	{                                                                             \
+		Component->PropName = PropName;                                           \
+	}                                                                             \
+}
+
+#define IMPLEMENT_GET_DYNAMICS_CONSTRAINT(PropName)                               \
+FRigUnit_HierarchyGetDynamicsConstraint##PropName##_Execute()                     \
+{                                                                                 \
+	if (!ExecuteContext.Hierarchy) { return; }                                    \
+	if (const FRigDynamicsConstraintComponent* Component =                        \
+		GetConstraint(*ExecuteContext.Hierarchy, DynamicsConstraintComponentKey)) \
+	{                                                                             \
+		PropName = Component->PropName;                                           \
+	}                                                                             \
+}
+
+//======================================================================================================================
+FRigUnit_SpawnDynamicsConstraint_Execute()
+{
+	if (!ExecuteContext.Hierarchy)
+	{
+		return;
+	}
+	if (ExecuteContext.GetEventName() != FRigUnit_PrepareForExecution::EventName)
+	{
+		UE_CONTROLRIG_RIGUNIT_REPORT_ERROR(TEXT("SpawnDynamicsConstraint can only be used during Setup"));
+		return;
+	}
+
+	if (URigHierarchyController* Controller = ExecuteContext.Hierarchy->GetController())
+	{
+		FRigHierarchyControllerInstructionBracket InstructionBracket(
+			Controller, ExecuteContext.GetInstructionIndex());
+		DynamicsConstraintComponentKey = Controller->AddComponent(
+			FRigDynamicsConstraintComponent::StaticStruct(), FRigDynamicsConstraintComponent::GetDefaultName(), Owner);
+		if (DynamicsConstraintComponentKey.IsValid())
+		{
+			if (FRigDynamicsConstraintComponent* Component =
+				GetConstraint(*ExecuteContext.Hierarchy, DynamicsConstraintComponentKey))
+			{
+				Component->ParentComponentKey = ParentComponentKey;
+				Component->ChildComponentKey = ChildComponentKey;
+				Component->ConstraintType = ConstraintType;
+				Component->Strength = Strength;
+				Component->DampingRatio = DampingRatio;
+				Component->ExtraDamping = ExtraDamping;
+				Component->LengthMultiplier = LengthMultiplier;
+				Component->ExtraLength = ExtraLength;
+			}
+
+			if (FRigDynamicsSolverComponent* Solver = GetSolver(*ExecuteContext.Hierarchy, DynamicsSolverComponentKey))
+			{
+				Solver->Constraints.Add(DynamicsConstraintComponentKey);
+			}
+		}
+	}
+}
+
+//======================================================================================================================
+IMPLEMENT_SET_DYNAMICS_CONSTRAINT(Strength)
+IMPLEMENT_GET_DYNAMICS_CONSTRAINT(Strength)
+//======================================================================================================================
+IMPLEMENT_SET_DYNAMICS_CONSTRAINT(DampingRatio)
+IMPLEMENT_GET_DYNAMICS_CONSTRAINT(DampingRatio)
+//======================================================================================================================
+IMPLEMENT_SET_DYNAMICS_CONSTRAINT(ExtraDamping)
+IMPLEMENT_GET_DYNAMICS_CONSTRAINT(ExtraDamping)
+//======================================================================================================================
+IMPLEMENT_SET_DYNAMICS_CONSTRAINT(LengthMultiplier)
+IMPLEMENT_GET_DYNAMICS_CONSTRAINT(LengthMultiplier)
+//======================================================================================================================
+IMPLEMENT_SET_DYNAMICS_CONSTRAINT(ExtraLength)
+IMPLEMENT_GET_DYNAMICS_CONSTRAINT(ExtraLength)
diff --git a/Engine/Plugins/Experimental/ControlRigDynamics/Source/ControlRigDynamics/Private/RigDynamicsData.cpp b/Engine/Plugins/Experimental/ControlRigDynamics/Source/ControlRigDynamics/Private/RigDynamicsData.cpp
new file mode 100644
--- /dev/null
+++ b/Engine/Plugins/Experimental/ControlRigDynamics/Source/ControlRigDynamics/Private/RigDynamicsData.cpp
@@ -0,0 +1,114 @@
+// Copyright Epic Games, Inc. All Rights Reserved.
+
+#include "RigDynamicsData.h"
+
+#include "RigDynamicsObjectVersion.h"
+
+//======================================================================================================================
+FArchive& operator <<(FArchive& Ar, FRigDynamicsColliderBox& Data)
+{
+	Ar << Data.Name;
+	Ar << Data.TM;
+	Ar << Data.Extents;
+	return Ar;
+}
+
+//======================================================================================================================
+FArchive& operator <<(FArchive& Ar, FRigDynamicsColliderCapsule& Data)
+{
+	Ar << Data.Name;
+	Ar << Data.TM;
+	Ar << Data.Radius;
+	Ar << Data.Length;
+	return Ar;
+}
+
+//======================================================================================================================
+FArchive& operator <<(FArchive& Ar, FRigDynamicsColliderPlane& Data)
+{
+	Ar << Data.Name;
+	if (Ar.CustomVer(FRigDynamicsObjectVersion::GUID) >= FRigDynamicsObjectVersion::PlaneDefinition)
+	{
+		Ar << Data.TM;
+	}
+	else
+	{
+		FVector Junk;
+		Ar << Junk;
+		Ar << Junk;
+	}
+	if (Ar.CustomVer(FRigDynamicsObjectVersion::GUID) >= FRigDynamicsObjectVersion::PlaneExtents)
+	{
+		Ar << Data.Extents;
+	}
+	return Ar;
+}
+
+//======================================================================================================================
+FArchive& operator <<(FArchive& Ar, FRigDynamicsColliderCollection& Data)
+{
+	Ar << Data.Boxes;
+	Ar << Data.Capsules;
+	Ar << Data.Planes;
+	return Ar;
+}
+
+//======================================================================================================================
+FArchive& operator <<(FArchive& Ar, FRigDynamicsSimulationSpaceSettings& Data)
+{
+	Ar << Data.SpaceMovementAmount;
+	Ar << Data.VelocityScaleZ;
+	Ar << Data.bClampLinearVelocity;
+	Ar << Data.MaxLinearVelocity;
+	Ar << Data.bClampAngularVelocity;
+	Ar << Data.MaxAngularVelocity;
+	Ar << Data.bClampLinearAcceleration;
+	Ar << Data.MaxLinearAcceleration;
+	Ar << Data.bClampAngularAcceleration;
+	Ar << Data.MaxAngularAcceleration;
+	Ar << Data.LinearAccelerationThresholdForTeleport;
+	Ar << Data.AngularAccelerationThresholdForTeleport;
+	Ar << Data.PositionChangeThresholdForTeleport;
+	Ar << Data.OrientationChangeThresholdForTeleport;
+	Ar << Data.LinearDragMultiplier;
+	Ar << Data.AngularDragMultiplier;
+	Ar << Data.ExternalLinearDrag;
+	Ar << Data.ExternalLinearVelocity;
+	Ar << Data.ExternalAngularVelocity;
+	Ar << Data.ExternalTurbulenceVelocity;
+	return Ar;
+}
+
+//======================================================================================================================
+FArchive& operator <<(FArchive& Ar, FRigDynamicsSolverSettings& Data)
+{
+	Ar << Data.Gravity;
+	Ar << Data.MaxTimeStep;
+	Ar << Data.MaxNumSteps;
+	Ar << Data.NumIterations;
+	Ar << Data.NumConstraintSubIterations;
+	if (Ar.CustomVer(FRigDynamicsObjectVersion::GUID) >= FRigDynamicsObjectVersion::BonePositionAndOrientationSetting)
+	{
+		Ar << Data.bSetBoneOrientations;
+		Ar << Data.bSetBonePositions;
+	}
+	if (Ar.CustomVer(FRigDynamicsObjectVersion::GUID) >= FRigDynamicsObjectVersion::SimulationSpace)
+	{
+		Ar << Data.SimulationSpace;
+		Ar << Data.CollisionSpace;
+		Ar << Data.SpaceBone;
+		Ar << Data.Colliders;
+	}
+	return Ar;
+}
+
+//======================================================================================================================
+FArchive& operator <<(FArchive& Ar, FRigDynamicsVisualizationSettings& Data)
+{
+	Ar << Data.bEnableVisualization;
+	Ar << Data.LineThickness;
+	Ar << Data.ShapeDetail;
+	Ar << Data.bShowParticles;
+	Ar << Data.bShowColliders;
+	return Ar;
+}
diff --git a/Engine/Plugins/Experimental/ControlRigDynamics/Source/ControlRigDynamics/Private/RigDynamicsExecution.cpp b/Engine/Plugins/Experimental/ControlRigDynamics/Source/ControlRigDynamics/Private/RigDynamicsExecution.cpp
new file mode 100644
--- /dev/null
+++ b/Engine/Plugins/Experimental/ControlRigDynamics/Source/ControlRigDynamics/Private/RigDynamicsExecution.cpp
@@ -0,0 +1,6 @@
+// Copyright Epic Games, Inc. All Rights Reserved.
+
+#include "RigDynamicsExecution.h"
+
+#include UE_INLINE_GENERATED_CPP_BY_NAME(RigDynamicsExecution)
+
diff --git a/Engine/Plugins/Experimental/ControlRigDynamics/Source/ControlRigDynamics/Private/RigDynamicsExecutionHelpers.cpp b/Engine/Plugins/Experimental/ControlRigDynamics/Source/ControlRigDynamics/Private/RigDynamicsExecutionHelpers.cpp
new file mode 100644
--- /dev/null
+++ b/Engine/Plugins/Experimental/ControlRigDynamics/Source/ControlRigDynamics/Private/RigDynamicsExecutionHelpers.cpp
@@ -0,0 +1,242 @@
+// Copyright Epic Games, Inc. All Rights Reserved.
+
+#include "RigDynamicsExecutionHelpers.h"
+
+#include "RigDynamicsHelpers.h"
+
+#include "Rigs/RigHierarchyController.h"
+#include "Units/RigUnitContext.h"
+#include "Units/Execution/RigUnit_PrepareForExecution.h"
+
+#include "PhysicsEngine/PhysicsAsset.h"
+#include "PhysicsEngine/SkeletalBodySetup.h"
+#include "PhysicsEngine/PhysicsConstraintTemplate.h"
+
+#include UE_INLINE_GENERATED_CPP_BY_NAME(RigDynamicsExecutionHelpers)
+
+//======================================================================================================================
+FRigUnit_SpawnDynamicsChains_Execute()
+{
+	if (ExecuteContext.GetEventName() != FRigUnit_PrepareForExecution::EventName)
+	{
+		UE_CONTROLRIG_RIGUNIT_REPORT_ERROR(TEXT("SpawnDynamicsChains can only be used during Setup"));
+		return;
+	}
+
+	if (!ExecuteContext.Hierarchy)
+	{
+		UE_CONTROLRIG_RIGUNIT_REPORT_ERROR(TEXT("SpawnDynamicsChains: No Hierarchy"));
+		return;
+	}
+	URigHierarchy& Hierarchy = *ExecuteContext.Hierarchy;
+
+	URigHierarchyController* Controller = ExecuteContext.Hierarchy->GetController();
+	if (!Controller)
+	{
+		UE_CONTROLRIG_RIGUNIT_REPORT_ERROR(TEXT("SpawnDynamicsChains: No Controller"));
+		return;
+	}
+
+	FRigHierarchyControllerInstructionBracket InstructionBracket(Controller, ExecuteContext.GetInstructionIndex());
+
+	// Get or create the solver
+	FRigDynamicsSolverComponent* SolverComponent = GetSolver(Hierarchy, FRigComponentKey(SolverElementKey, SolverComponentName));
+	if (SolverComponent)
+	{
+		SolverComponentKey = SolverComponent->GetKey();
+	}
+	else
+	{
+		SolverComponentKey = Controller->AddComponent(
+			FRigDynamicsSolverComponent::StaticStruct(), SolverComponentName, SolverElementKey);
+		if (SolverComponentKey.IsValid())
+		{
+			SolverComponent = GetSolver(*ExecuteContext.Hierarchy, SolverComponentKey);
+			if (SolverComponent)
+			{
+				SolverComponent->Settings = SolverSettings;
+				SolverComponent->SimulationSpaceSettings = SimulationSpaceSettings;
+			}
+		}
+	}
+
+	// Note that the key here may not be valid - but that's OK (the user may be going to assign them later)
+
+	// The array of elements for which we want to create particles 
+	TArray<FRigElementKey> ParticleElementKeys;
+
+	TSet<FRigElementKey> ExcludeElementKeys;
+	for (const FRigElementKey& ElementKey : ChainTerminatorBones)
+	{
+		TArray<FRigElementKey> Children = Hierarchy.GetChildren(ElementKey, true);
+		ExcludeElementKeys.Append(Children);
+	}
+
+	// In the end, the solver will sort constraints and particles. 
+	for (const FRigElementKey& ElementKey : RootBones)
+	{
+		ParticleElementKeys.Add(ElementKey);
+
+		TArray<FRigElementKey> Children = Hierarchy.GetChildren(ElementKey, true);
+
+		for (const FRigElementKey& Child : Children)
+		{
+			if (!ExcludeElementKeys.Contains(Child))
+			{
+				ParticleElementKeys.AddUnique(Child);
+			}
+		}
+	}
+
+	for (const FRigElementKey& ElementKey : ExtraBones)
+	{
+		ParticleElementKeys.AddUnique(ElementKey);
+	}
+
+	// Now we have ParticleElementKeys
+
+	for (const FRigElementKey& ElementKey : ParticleElementKeys)
+	{
+		FRigComponentKey DynamicsParticleComponentKey = Controller->AddComponent(
+			FRigDynamicsParticleComponent::StaticStruct(), FRigDynamicsParticleComponent::GetDefaultName(), ElementKey);
+		if (DynamicsParticleComponentKey.IsValid())
+		{
+			if (FRigDynamicsParticleComponent* Component =
+				GetParticle(*ExecuteContext.Hierarchy, DynamicsParticleComponentKey))
+			{
+				Component->ParticleProperties = ParticleProperties;
+			}
+
+			if (SolverComponent)
+			{
+				SolverComponent->Particles.Add(DynamicsParticleComponentKey);
+			}
+			ParticleComponentKeys.Add(DynamicsParticleComponentKey);
+		}
+	}
+
+	for (const FRigDynamicsColliderCollectionWithKey& ColliderCollectionWithKey : Colliders)
+	{
+		const FRigElementKey& OwnerKey = ColliderCollectionWithKey.Owner;
+		const FRigDynamicsColliderCollection& ColliderCollection = ColliderCollectionWithKey.ColliderCollection;
+
+		FRigComponentKey ColliderComponentKey = Controller->AddComponent(
+			FRigDynamicsColliderComponent::StaticStruct(), FRigDynamicsColliderComponent::GetDefaultName(), OwnerKey);
+		if (ColliderComponentKey.IsValid())
+		{
+			if (FRigDynamicsColliderComponent* Component = GetCollider(*ExecuteContext.Hierarchy, ColliderComponentKey))
+			{
+				Component->Colliders.Boxes = ColliderCollection.Boxes;
+				Component->Colliders.Capsules = ColliderCollection.Capsules;
+				Component->Colliders.Planes = ColliderCollection.Planes;
+			}
+
+			if (SolverComponent)
+			{
+				SolverComponent->Colliders.Add(ColliderComponentKey);
+			}
+
+			ColliderComponentKeys.Add(ColliderComponentKey);
+		}
+	}
+
+}
+
+//======================================================================================================================
+FRigUnit_HierarchyImportDynamicsCollidersFromPhysicsAsset_Execute()
+{
+	if (ExecuteContext.GetEventName() != FRigUnit_PrepareForExecution::EventName)
+	{
+		UE_CONTROLRIG_RIGUNIT_REPORT_ERROR(TEXT("ImportDynamicsCollidersFromPhysicsAsset can only be used during Setup"));
+		return;
+	}
+
+	if (!ExecuteContext.Hierarchy)
+	{
+		UE_CONTROLRIG_RIGUNIT_REPORT_ERROR(TEXT("ImportDynamicsCollidersFromPhysicsAsset: No Hierarchy"));
+		return;
+	}
+	URigHierarchy& Hierarchy = *ExecuteContext.Hierarchy;
+
+	if (!PhysicsAsset)
+	{
+		UE_CONTROLRIG_RIGUNIT_REPORT_ERROR(TEXT("ImportDynamicsCollidersFromPhysicsAsset needs a valid physics asset"));
+		return;
+	}
+
+	URigHierarchyController* Controller = Hierarchy.GetController();
+	if (!Controller)
+	{
+		UE_CONTROLRIG_RIGUNIT_REPORT_ERROR(TEXT("ImportDynamicsCollidersFromPhysicsAsset: No Controller"));
+		return;
+	}
+
+	FRigDynamicsSolverComponent* Solver = GetSolver(*ExecuteContext.Hierarchy, DynamicsSolverComponentKey);
+
+	FRigHierarchyControllerInstructionBracket InstructionBracket(Controller, ExecuteContext.GetInstructionIndex());
+
+	for (TObjectPtr<USkeletalBodySetup> SkeletalBodySetup : PhysicsAsset->SkeletalBodySetups)
+	{
+		if (!SkeletalBodySetup)
+		{
+			continue;
+		}
+
+		// Check we can use this based on the bone name
+		FName BoneName = SkeletalBodySetup->BoneName;
+		const FRigElementKey OwnerElementKey(BoneName, ERigElementType::Bone);
+		if (!ExecuteContext.Hierarchy->Find(OwnerElementKey))
+		{
+			continue;
+		}
+		if (!BoneMask.IsEmpty() && BoneMask.Find(BoneName) == INDEX_NONE)
+		{
+			continue;
+		}
+
+		FRigComponentKey DynamicsColliderComponentKey = Controller->AddComponent(
+			FRigDynamicsColliderComponent::StaticStruct(),
+			FRigDynamicsColliderComponent::GetDefaultName(),
+			OwnerElementKey);
+
+		if (DynamicsColliderComponentKey.IsValid())
+		{
+			if (FRigDynamicsColliderComponent* Component =
+				GetCollider(*ExecuteContext.Hierarchy, DynamicsColliderComponentKey))
+			{
+				DynamicsColliderComponentKeys.Add(DynamicsColliderComponentKey);
+
+				FRigDynamicsColliderCollection& ColliderCollection = Component->Colliders;
+
+				const FKAggregateGeom& AggGeom = SkeletalBodySetup->AggGeom;
+				for (const FKBoxElem& BoxElem : AggGeom.BoxElems)
+				{
+					FTransform TM(BoxElem.Rotation, BoxElem.Center);
+					FRigDynamicsColliderBox Box(
+						FName(SkeletalBodySetup->GetName()), TM, FVector(BoxElem.X, BoxElem.Y, BoxElem.Z));
+					ColliderCollection.Boxes.Add(Box);
+				}
+				for (const FKSphereElem& SphereElem : AggGeom.SphereElems)
+				{
+					FTransform TM(SphereElem.Center);
+					FRigDynamicsColliderCapsule Capsule(
+						FName(SkeletalBodySetup->GetName()), TM, SphereElem.Radius, 0.0f);
+					ColliderCollection.Capsules.Add(Capsule);
+				}
+				for (const FKSphylElem& SphylElem : AggGeom.SphylElems)
+				{
+					FTransform TM(SphylElem.Rotation, SphylElem.Center);
+					FRigDynamicsColliderCapsule Capsule(
+						FName(SkeletalBodySetup->GetName()), TM, SphylElem.Radius, SphylElem.Length);
+					ColliderCollection.Capsules.Add(Capsule);
+				}
+
+				if (Solver)
+				{
+					Solver->Colliders.Add(DynamicsColliderComponentKey);
+				}
+			}
+		}
+
+	}
+}
diff --git a/Engine/Plugins/Experimental/ControlRigDynamics/Source/ControlRigDynamics/Private/RigDynamicsHelpers.h b/Engine/Plugins/Experimental/ControlRigDynamics/Source/ControlRigDynamics/Private/RigDynamicsHelpers.h
new file mode 100644
--- /dev/null
+++ b/Engine/Plugins/Experimental/ControlRigDynamics/Source/ControlRigDynamics/Private/RigDynamicsHelpers.h
@@ -0,0 +1,126 @@
+// Copyright Epic Games, Inc. All Rights Reserved.
+
+#pragma once
+
+#include "RigDynamicsParticleComponent.h"
+#include "RigDynamicsSolverComponent.h"
+#include "RigDynamicsColliderComponent.h"
+#include "RigDynamicsConstraintComponent.h"
+
+#include "Rigs/RigHierarchy.h"
+
+//======================================================================================================================
+inline FRigDynamicsParticleComponent* GetParticle(URigHierarchy& Hierarchy, const FRigComponentKey& Key)
+{
+	return Cast<FRigDynamicsParticleComponent>(Hierarchy.FindComponent(Key));
+}
+
+//======================================================================================================================
+inline const FRigDynamicsParticleComponent* GetParticle(const URigHierarchy& Hierarchy, const FRigComponentKey& Key)
+{
+	return Cast<FRigDynamicsParticleComponent>(Hierarchy.FindComponent(Key));
+}
+
+//======================================================================================================================
+// Updates the cache if necessary and then returns the component
+inline const FRigDynamicsParticleComponent* GetParticle(
+	const URigHierarchy& Hierarchy, FCachedRigComponent& CachedComponent)
+{
+	return Cast<FRigDynamicsParticleComponent>(CachedComponent.GetComponent(&Hierarchy));
+}
+
+//======================================================================================================================
+inline FRigDynamicsSolverComponent* GetSolver(URigHierarchy& Hierarchy, const FRigComponentKey& Key)
+{
+	return Cast<FRigDynamicsSolverComponent>(Hierarchy.FindComponent(Key));
+}
+
+//======================================================================================================================
+inline const FRigDynamicsSolverComponent* GetSolver(const URigHierarchy& Hierarchy, const FRigComponentKey& Key)
+{
+	return Cast<FRigDynamicsSolverComponent>(Hierarchy.FindComponent(Key));
+}
+
+//======================================================================================================================
+// Updates the cache if necessary and then returns the component
+inline const FRigDynamicsSolverComponent* GetSolver(
+	const URigHierarchy& Hierarchy, FCachedRigComponent& CachedComponent)
+{
+	return Cast<FRigDynamicsSolverComponent>(CachedComponent.GetComponent(&Hierarchy));
+}
+
+//======================================================================================================================
+inline FRigDynamicsColliderComponent* GetCollider(URigHierarchy& Hierarchy, const FRigComponentKey& Key)
+{
+	return Cast<FRigDynamicsColliderComponent>(Hierarchy.FindComponent(Key));
+}
+
+//======================================================================================================================
+inline const FRigDynamicsColliderComponent* GetCollider(const URigHierarchy& Hierarchy, const FRigComponentKey& Key)
+{
+	return Cast<FRigDynamicsColliderComponent>(Hierarchy.FindComponent(Key));
+}
+
+//======================================================================================================================
+// Updates the cache if necessary and then returns the component
+inline const FRigDynamicsColliderComponent* GetCollider(
+	const URigHierarchy& Hierarchy, FCachedRigComponent& CachedComponent)
+{
+	return Cast<FRigDynamicsColliderComponent>(CachedComponent.GetComponent(&Hierarchy));
+}
+
+//======================================================================================================================
+inline FRigDynamicsConstraintComponent* GetConstraint(URigHierarchy& Hierarchy, const FRigComponentKey& Key)
+{
+	return Cast<FRigDynamicsConstraintComponent>(Hierarchy.FindComponent(Key));
+}
+
+//======================================================================================================================
+inline const FRigDynamicsConstraintComponent* GetConstraint(const URigHierarchy& Hierarchy, const FRigComponentKey& Key)
+{
+	return Cast<FRigDynamicsConstraintComponent>(Hierarchy.FindComponent(Key));
+}
+
+//======================================================================================================================
+// Updates the cache if necessary and then returns the component
+inline const FRigDynamicsConstraintComponent* GetConstraint(
+	const URigHierarchy& Hierarchy, FCachedRigComponent& CachedRigComponent)
+{
+	return Cast<FRigDynamicsConstraintComponent>(CachedRigComponent.GetComponent(&Hierarchy));
+}
+
+//======================================================================================================================
+// Gets the transform from the element that owns this component. 
+inline FTransform GetGlobalTransform(const URigHierarchy& Hierarchy, const FRigBaseComponent& Component)
+{
+	// Since all components are owned, they must have an element
+	check(Component.GetElement());
+	return Hierarchy.GetGlobalTransform(Component.GetElement()->GetIndex());
+}
+
+//======================================================================================================================
+// Gets the location from the element that owns this component. 
+inline FVector GetGlobalLocation(const URigHierarchy& Hierarchy, const FRigBaseComponent& Component)
+{
+	// Since all components are owned, they must have an element
+	check(Component.GetElement());
+	return Hierarchy.GetGlobalTransform(Component.GetElement()->GetIndex()).GetLocation();
+}
+
+//======================================================================================================================
+// We assume that the cached rig component is fundamentally valid (i.e. exists), but we do ensure it is cached.
+inline int32 GetElementIndex(URigHierarchy& Hierarchy, FCachedRigComponent& CachedRigComponent)
+{
+	CachedRigComponent.UpdateCache(&Hierarchy);
+	const FRigBaseElement* Element = CachedRigComponent.GetElement();
+	// Since all components are owned, they must have an element
+	check(Element);
+	return Element->GetIndex();
+}
+
+//======================================================================================================================
+inline FTransform GetGlobalTransform(URigHierarchy& Hierarchy, FCachedRigComponent& CachedRigComponent)
+{
+	const int32 Index = GetElementIndex(Hierarchy, CachedRigComponent);
+	return Hierarchy.GetGlobalTransform(Index);
+}
\ No newline at end of file
diff --git a/Engine/Plugins/Experimental/ControlRigDynamics/Source/ControlRigDynamics/Private/RigDynamicsObjectVersion.h b/Engine/Plugins/Experimental/ControlRigDynamics/Source/ControlRigDynamics/Private/RigDynamicsObjectVersion.h
new file mode 100644
--- /dev/null
+++ b/Engine/Plugins/Experimental/ControlRigDynamics/Source/ControlRigDynamics/Private/RigDynamicsObjectVersion.h
@@ -0,0 +1,40 @@
+// Copyright Epic Games, Inc. All Rights Reserved.
+
+#pragma once
+
+#include "Misc/Guid.h"
+
+#define UE_API CONTROLRIGDYNAMICS_API
+
+//======================================================================================================================
+struct FRigDynamicsObjectVersion
+{
+	enum Type
+	{
+		FirstVersion,
+		ParticleExtraDamping,
+		DynamicsTargetMode,
+		GravityMultiplier,
+		PlaneDefinition,
+		PlaneExtents,
+		Constraints,
+		HelperStructs,
+		BonePositionAndOrientationSetting,
+		NoCollisionColliders,
+		CollisionParticles,
+		TargetModeFloat,
+		SimulationSpace,
+
+		// -----<new versions can be added above this line>-----
+		VersionPlusOne,
+		LatestVersion = VersionPlusOne - 1
+	};
+
+	// The GUID for this custom version number
+	UE_API const static FGuid GUID;
+
+private:
+	FRigDynamicsObjectVersion() {}
+};
+
+#undef UE_API
diff --git a/Engine/Plugins/Experimental/ControlRigDynamics/Source/ControlRigDynamics/Private/RigDynamicsParticleComponent.cpp b/Engine/Plugins/Experimental/ControlRigDynamics/Source/ControlRigDynamics/Private/RigDynamicsParticleComponent.cpp
new file mode 100644
--- /dev/null
+++ b/Engine/Plugins/Experimental/ControlRigDynamics/Source/ControlRigDynamics/Private/RigDynamicsParticleComponent.cpp
@@ -0,0 +1,86 @@
+// Copyright Epic Games, Inc. All Rights Reserved.
+
+#include "RigDynamicsParticleComponent.h"
+
+#include "RigDynamicsObjectVersion.h"
+
+//======================================================================================================================
+FArchive& operator <<(FArchive& Ar, FRigDynamicsParticleProperties& Data)
+{
+	Ar << Data.Radius;
+	Ar << Data.Mass;
+	Ar << Data.MovementType;
+	Ar << Data.GravityMultiplier;
+	Ar << Data.Strength;
+	Ar << Data.DampingRatio;
+	Ar << Data.ExtraDamping;
+	Ar << Data.TargetVelocityInfluence;
+	if (Ar.CustomVer(FRigDynamicsObjectVersion::GUID) >= FRigDynamicsObjectVersion::TargetModeFloat)
+	{
+		Ar << Data.TargetMode;
+	}
+	else
+	{
+		// Old data: TargetMode was ERigParticleSimulationTargetMode (uint8).
+		// SimSpace=0 -> 0.0f, PreservePose=1 -> 1.0f
+		uint8 LegacyTargetMode = 1;
+		Ar << LegacyTargetMode;
+	}
+	if (Ar.CustomVer(FRigDynamicsObjectVersion::GUID) >= FRigDynamicsObjectVersion::NoCollisionColliders)
+	{
+		Ar << Data.NoCollisionColliders;
+	}
+	if (Ar.CustomVer(FRigDynamicsObjectVersion::GUID) >= FRigDynamicsObjectVersion::CollisionParticles)
+	{
+		Ar << Data.CollisionParticles;
+	}
+	return Ar;
+}
+
+
+//======================================================================================================================
+void FRigDynamicsParticleComponent::Serialize(FArchive& Ar)
+{
+	if (Ar.CustomVer(FRigDynamicsObjectVersion::GUID) >= FRigDynamicsObjectVersion::HelperStructs)
+	{
+		Ar << ParticleProperties;
+	}
+	else
+	{
+		Ar << ParticleProperties.Radius;
+		Ar << ParticleProperties.Mass;
+		Ar << ParticleProperties.MovementType;
+		if (Ar.CustomVer(FRigDynamicsObjectVersion::GUID) >= FRigDynamicsObjectVersion::GravityMultiplier)
+		{
+			Ar << ParticleProperties.GravityMultiplier;
+		}
+		Ar << ParticleProperties.Strength;
+		Ar << ParticleProperties.DampingRatio;
+		if (Ar.CustomVer(FRigDynamicsObjectVersion::GUID) >= FRigDynamicsObjectVersion::ParticleExtraDamping)
+		{
+			Ar << ParticleProperties.ExtraDamping;
+		}
+		Ar << ParticleProperties.TargetVelocityInfluence;
+		if (Ar.CustomVer(FRigDynamicsObjectVersion::GUID) >= FRigDynamicsObjectVersion::DynamicsTargetMode)
+		{
+			// Pre-HelperStructs is always pre-TargetModeFloat, so always read the legacy uint8
+			uint8 LegacyTargetMode = 1;
+			Ar << LegacyTargetMode;
+		}
+	}
+}
+
+//======================================================================================================================
+void FRigDynamicsParticleComponent::Save(FArchive& Ar)
+{
+	Ar.UsingCustomVersion(FRigDynamicsObjectVersion::GUID);
+	FRigBaseComponent::Save(Ar);
+	Serialize(Ar);
+}
+
+//======================================================================================================================
+void FRigDynamicsParticleComponent::Load(FArchive& Ar)
+{
+	FRigBaseComponent::Load(Ar);
+	Serialize(Ar);
+}
diff --git a/Engine/Plugins/Experimental/ControlRigDynamics/Source/ControlRigDynamics/Private/RigDynamicsParticleExecution.cpp b/Engine/Plugins/Experimental/ControlRigDynamics/Source/ControlRigDynamics/Private/RigDynamicsParticleExecution.cpp
new file mode 100644
--- /dev/null
+++ b/Engine/Plugins/Experimental/ControlRigDynamics/Source/ControlRigDynamics/Private/RigDynamicsParticleExecution.cpp
@@ -0,0 +1,150 @@
+// Copyright Epic Games, Inc. All Rights Reserved.
+
+#include "RigDynamicsParticleExecution.h"
+
+#include "RigDynamicsHelpers.h"
+
+#include "Rigs/RigHierarchyController.h"
+#include "Units/RigUnitContext.h"
+#include "Units/Execution/RigUnit_PrepareForExecution.h"
+
+#include UE_INLINE_GENERATED_CPP_BY_NAME(RigDynamicsParticleExecution)
+
+//======================================================================================================================
+// Macros to reduce boilerplate for the Get/Set particle property nodes. Note that these can't
+// be used in the header because UHT scans raw source before C++ preprocessing, so it would
+// never see the USTRUCT/UPROPERTY/RIGVM_METHOD markers inside a macro expansion.
+//======================================================================================================================
+
+#define IMPLEMENT_SET_DYNAMICS_PARTICLE(PropName)                                \
+FRigUnit_HierarchySetDynamicsParticle##PropName##_Execute()                      \
+{                                                                                \
+	if (!ExecuteContext.Hierarchy) { return; }                                   \
+	if (FRigDynamicsParticleComponent* Component =                               \
+		GetParticle(*ExecuteContext.Hierarchy, DynamicsParticleComponentKey))    \
+	{                                                                            \
+		Component->ParticleProperties.PropName = PropName;                       \
+	}                                                                            \
+}
+
+#define IMPLEMENT_GET_DYNAMICS_PARTICLE(PropName)                                \
+FRigUnit_HierarchyGetDynamicsParticle##PropName##_Execute()                      \
+{                                                                                \
+	if (!ExecuteContext.Hierarchy) { return; }                                   \
+	if (const FRigDynamicsParticleComponent* Component =                         \
+		GetParticle(*ExecuteContext.Hierarchy, DynamicsParticleComponentKey))    \
+	{                                                                            \
+		PropName = Component->ParticleProperties.PropName;                       \
+	}                                                                            \
+}
+
+//======================================================================================================================
+FRigUnit_SpawnDynamicsParticle_Execute()
+{
+	if (!ExecuteContext.Hierarchy)
+	{
+		return;
+	}
+	if (ExecuteContext.GetEventName() != FRigUnit_PrepareForExecution::EventName)
+	{
+		UE_CONTROLRIG_RIGUNIT_REPORT_ERROR(TEXT("SpawnDynamicsParticle can only be used during Setup"));
+		return;
+	}
+
+	if (URigHierarchyController* Controller = ExecuteContext.Hierarchy->GetController())
+	{
+		FRigHierarchyControllerInstructionBracket InstructionBracket(
+			Controller, ExecuteContext.GetInstructionIndex());
+		DynamicsParticleComponentKey = Controller->AddComponent(
+			FRigDynamicsParticleComponent::StaticStruct(), FRigDynamicsParticleComponent::GetDefaultName(), Owner);
+		if (DynamicsParticleComponentKey.IsValid())
+		{
+			if (FRigDynamicsParticleComponent* Component =
+				GetParticle(*ExecuteContext.Hierarchy, DynamicsParticleComponentKey))
+			{
+				Component->ParticleProperties = ParticleProperties;
+			}
+
+			if (FRigDynamicsSolverComponent* Solver = GetSolver(*ExecuteContext.Hierarchy, DynamicsSolverComponentKey))
+			{
+				Solver->Particles.Add(DynamicsParticleComponentKey);
+			}
+		}
+	}
+}
+
+//======================================================================================================================
+FRigUnit_HierarchyDisableDynamicsCollisionWithCollider_Execute()
+{
+	if (!ExecuteContext.Hierarchy)
+	{
+		return;
+	}
+	if (ExecuteContext.GetEventName() != FRigUnit_PrepareForExecution::EventName)
+	{
+		UE_CONTROLRIG_RIGUNIT_REPORT_ERROR(TEXT("DisableDynamicsCollisionWithCollider can only be used during Setup"));
+		return;
+	}
+
+	if (FRigDynamicsParticleComponent* ParticleComponent = 
+		GetParticle(*ExecuteContext.Hierarchy, DynamicsParticleComponentKey))
+	{
+		ParticleComponent->ParticleProperties.NoCollisionColliders.AddUnique(DynamicsColliderComponentKey);
+	}
+	else
+	{
+		UE_CONTROLRIG_RIGUNIT_REPORT_WARNING(TEXT("DisableDynamicsCollisionWithCollider Unable to find particle"));
+	}
+}
+
+//======================================================================================================================
+FRigUnit_HierarchyEnableDynamicsCollisionWithParticle_Execute()
+{
+	if (!ExecuteContext.Hierarchy)
+	{
+		return;
+	}
+	if (ExecuteContext.GetEventName() != FRigUnit_PrepareForExecution::EventName)
+	{
+		UE_CONTROLRIG_RIGUNIT_REPORT_ERROR(TEXT("EnableDynamicsCollisionWithParticle can only be used during Setup"));
+		return;
+	}
+
+	if (FRigDynamicsParticleComponent* ParticleComponent = 
+		GetParticle(*ExecuteContext.Hierarchy, DynamicsParticleComponentKey))
+	{
+		ParticleComponent->ParticleProperties.CollisionParticles.AddUnique(OtherDynamicsParticleComponentKey);
+	}
+	else
+	{
+		UE_CONTROLRIG_RIGUNIT_REPORT_WARNING(TEXT("EnableDynamicsCollisionWithParticle Unable to find particle"));
+	}
+}
+
+//======================================================================================================================
+IMPLEMENT_SET_DYNAMICS_PARTICLE(Radius)
+IMPLEMENT_GET_DYNAMICS_PARTICLE(Radius)
+//======================================================================================================================
+IMPLEMENT_SET_DYNAMICS_PARTICLE(Mass)
+IMPLEMENT_GET_DYNAMICS_PARTICLE(Mass)
+//======================================================================================================================
+IMPLEMENT_SET_DYNAMICS_PARTICLE(MovementType)
+IMPLEMENT_GET_DYNAMICS_PARTICLE(MovementType)
+//======================================================================================================================
+IMPLEMENT_SET_DYNAMICS_PARTICLE(GravityMultiplier)
+IMPLEMENT_GET_DYNAMICS_PARTICLE(GravityMultiplier)
+//======================================================================================================================
+IMPLEMENT_SET_DYNAMICS_PARTICLE(Strength)
+IMPLEMENT_GET_DYNAMICS_PARTICLE(Strength)
+//======================================================================================================================
+IMPLEMENT_SET_DYNAMICS_PARTICLE(DampingRatio)
+IMPLEMENT_GET_DYNAMICS_PARTICLE(DampingRatio)
+//======================================================================================================================
+IMPLEMENT_SET_DYNAMICS_PARTICLE(ExtraDamping)
+IMPLEMENT_GET_DYNAMICS_PARTICLE(ExtraDamping)
+//======================================================================================================================
+IMPLEMENT_SET_DYNAMICS_PARTICLE(TargetVelocityInfluence)
+IMPLEMENT_GET_DYNAMICS_PARTICLE(TargetVelocityInfluence)
+//======================================================================================================================
+IMPLEMENT_SET_DYNAMICS_PARTICLE(TargetMode)
+IMPLEMENT_GET_DYNAMICS_PARTICLE(TargetMode)
diff --git a/Engine/Plugins/Experimental/ControlRigDynamics/Source/ControlRigDynamics/Private/RigDynamicsSimulationSpace.cpp b/Engine/Plugins/Experimental/ControlRigDynamics/Source/ControlRigDynamics/Private/RigDynamicsSimulationSpace.cpp
new file mode 100644
--- /dev/null
+++ b/Engine/Plugins/Experimental/ControlRigDynamics/Source/ControlRigDynamics/Private/RigDynamicsSimulationSpace.cpp
@@ -0,0 +1,197 @@
+// Copyright Epic Games, Inc. All Rights Reserved.
+
+#include "RigDynamicsSimulationSpace.h"
+
+#include "PhysicsControlHelpers.h"
+
+//======================================================================================================================
+// TODO We need support for double precision Perlin inputs - for now duplicate.
+namespace FPerlinHelpers
+{
+	// random permutation of 256 numbers, repeated 2x
+	static const int32 Permutation[512] = {
+		63, 9, 212, 205, 31, 128, 72, 59, 137, 203, 195, 170, 181, 115, 165, 40, 116, 139, 175, 225, 132, 99, 
+		222, 2, 41, 15, 197, 93, 169, 90, 228, 43, 221, 38, 206, 204, 73, 17, 97, 10, 96, 47, 32, 138, 136, 30, 219,
+		78, 224, 13, 193, 88, 134, 211, 7, 112, 176, 19, 106, 83, 75, 217, 85, 0, 98, 140, 229, 80, 118, 151, 
+		117, 251, 103, 242, 81, 238, 172, 82, 110, 4, 227, 77, 243, 46, 12, 189, 34, 188, 200, 161, 68, 76, 171, 194,
+		57, 48, 247, 233, 51, 105, 5, 23, 42, 50, 216, 45, 239, 148, 249, 84, 70, 125, 108, 241, 62, 66, 64, 240, 
+		173, 185, 250, 49, 6, 37, 26, 21, 244, 60, 223, 255, 16, 145, 27, 109, 58, 102, 142, 253, 120, 149, 160,
+		124, 156, 79, 186, 135, 127, 14, 121, 22, 65, 54, 153, 91, 213, 174, 24, 252, 131, 192, 190, 202, 208, 35, 
+		94, 231, 56, 95, 183, 163, 111, 147, 25, 67, 36, 92, 236, 71, 166, 1, 187, 100, 130, 143, 237, 178, 158,
+		104, 184, 159, 177, 52, 214, 230, 119, 87, 114, 201, 179, 198, 3, 248, 182, 39, 11, 152, 196, 113, 20, 232, 
+		69, 141, 207, 234, 53, 86, 180, 226, 74, 150, 218, 29, 133, 8, 44, 123, 28, 146, 89, 101, 154, 220, 126,
+		155, 122, 210, 168, 254, 162, 129, 33, 18, 209, 61, 191, 199, 157, 245, 55, 164, 167, 215, 246, 144, 107, 235,
+
+		63, 9, 212, 205, 31, 128, 72, 59, 137, 203, 195, 170, 181, 115, 165, 40, 116, 139, 175, 225, 132, 99, 222, 
+		2, 41, 15, 197, 93, 169, 90, 228, 43, 221, 38, 206, 204, 73, 17, 97, 10, 96, 47, 32, 138, 136, 30, 219,
+		78, 224, 13, 193, 88, 134, 211, 7, 112, 176, 19, 106, 83, 75, 217, 85, 0, 98, 140, 229, 80, 118, 151, 117, 
+		251, 103, 242, 81, 238, 172, 82, 110, 4, 227, 77, 243, 46, 12, 189, 34, 188, 200, 161, 68, 76, 171, 194,
+		57, 48, 247, 233, 51, 105, 5, 23, 42, 50, 216, 45, 239, 148, 249, 84, 70, 125, 108, 241, 62, 66, 64, 240, 
+		173, 185, 250, 49, 6, 37, 26, 21, 244, 60, 223, 255, 16, 145, 27, 109, 58, 102, 142, 253, 120, 149, 160,
+		124, 156, 79, 186, 135, 127, 14, 121, 22, 65, 54, 153, 91, 213, 174, 24, 252, 131, 192, 190, 202, 208, 35, 
+		94, 231, 56, 95, 183, 163, 111, 147, 25, 67, 36, 92, 236, 71, 166, 1, 187, 100, 130, 143, 237, 178, 158,
+		104, 184, 159, 177, 52, 214, 230, 119, 87, 114, 201, 179, 198, 3, 248, 182, 39, 11, 152, 196, 113, 20, 232, 
+		69, 141, 207, 234, 53, 86, 180, 226, 74, 150, 218, 29, 133, 8, 44, 123, 28, 146, 89, 101, 154, 220, 126,
+		155, 122, 210, 168, 254, 162, 129, 33, 18, 209, 61, 191, 199, 157, 245, 55, 164, 167, 215, 246, 144, 107, 235
+	};
+
+	//==================================================================================================================
+	FORCEINLINE float Grad1(int32 Hash, double X)
+	{
+		static const double Grad1Scales[16] = {
+			-8 / 8., -7 / 8., -6 / 8., -5 / 8., -4 / 8., -3 / 8., -2 / 8., -1 / 8.,
+			1 / 8., 2 / 8., 3 / 8., 4 / 8., 5 / 8., 6 / 8., 7 / 8., 8 / 8. };
+		return Grad1Scales[Hash & 15] * X;
+	}
+
+	//==================================================================================================================
+	FORCEINLINE double SmoothCurve(double X)
+	{
+		return X * X * X * (X * (X * 6.0 - 15.0) + 10.0);
+	}
+
+	//==================================================================================================================
+	FORCEINLINE float PerlinNoise1D(double X)
+	{
+		const double Xfl = FMath::FloorToDouble(X);
+		const int64 Xi = (int64)(Xfl) & 255;
+		X -= Xfl;
+		const double Xm1 = X - 1.0;
+
+		const int32 A = Permutation[Xi];
+		const int32 B = Permutation[Xi + 1];
+
+		const double U = SmoothCurve(X);
+
+		// 2.0 factor to ensure (-1, 1) range
+		return 2.0f * FMath::Lerp(Grad1(A, X), Grad1(B, Xm1), U);
+	}
+}
+
+//======================================================================================================================
+void FRigDynamicsSimulationSpaceState::Update(
+	const FTransform& InComponentTM, const FTransform& InSimulationSpaceTM, float InDeltaTime)
+{
+	if (InDeltaTime > 0.0f)
+	{
+		FVector NewLinearVelocity = UE::PhysicsControl::CalculateLinearVelocity(
+			SimulationSpaceTM.GetTranslation(), InSimulationSpaceTM.GetTranslation(), InDeltaTime);
+		FVector NewAngularVelocity = UE::PhysicsControl::CalculateAngularVelocity(
+			SimulationSpaceTM.GetRotation(), InSimulationSpaceTM.GetRotation(), InDeltaTime);
+
+		if (UpdatesSinceReset == 0)
+		{
+			LinearAcceleration = FVector::ZeroVector;
+			AngularAcceleration = FVector::ZeroVector;
+		}
+		else
+		{
+			LinearAcceleration = (NewLinearVelocity - LinearVelocity) / InDeltaTime;
+			AngularAcceleration = (NewAngularVelocity - AngularVelocity) / InDeltaTime;
+		}
+
+		LinearVelocity = NewLinearVelocity;
+		AngularVelocity = NewAngularVelocity;
+
+		PrevPrevSimulationSpaceTM = PrevSimulationSpaceTM;
+		PrevSimulationSpaceTM = SimulationSpaceTM;
+
+		PrevDeltaTime = DeltaTime;
+		DeltaTime = InDeltaTime;
+
+		++UpdatesSinceReset;
+	}
+
+	// If delta time is zero just update the transform and don't worry too much about velocities being off.
+	SimulationSpaceTM = InSimulationSpaceTM;
+	ComponentTM = InComponentTM;
+}
+
+//======================================================================================================================
+void FRigDynamicsSimulationSpaceState::Reset(
+	const FTransform& InComponentTM, const FTransform& InSimulationSpaceTM)
+{
+	*this = FRigDynamicsSimulationSpaceState();
+	SimulationSpaceTM = InSimulationSpaceTM;
+	ComponentTM = InComponentTM;
+	UpdatesSinceReset = 0;
+}
+
+//======================================================================================================================
+RigParticleSimulation::FSimulationSpaceMotion FRigDynamicsSimulationSpaceState::CalculateMotion(
+	const FRigDynamicsSimulationSpaceSettings& Settings, const double AbsoluteTime) const
+{
+	using RigParticleSimulation::FSimVector;
+
+	// Start with world-space velocities and accelerations from the state
+	FVector WorldLinVel = LinearVelocity;
+	FVector WorldAngVel = AngularVelocity;
+	FVector WorldLinAccel = LinearAcceleration;
+	FVector WorldAngAccel = AngularAcceleration;
+
+	// Apply Z scale
+	WorldLinVel.Z *= Settings.VelocityScaleZ;
+	WorldLinAccel.Z *= Settings.VelocityScaleZ;
+
+	// TODO Need to handle teleport
+
+	// Clamp
+	if (Settings.bClampLinearVelocity)
+	{
+		WorldLinVel = WorldLinVel.GetClampedToMaxSize(Settings.MaxLinearVelocity);
+	}
+	if (Settings.bClampAngularVelocity)
+	{
+		WorldAngVel = WorldAngVel.GetClampedToMaxSize(Settings.MaxAngularVelocity);
+	}
+	if (Settings.bClampLinearAcceleration)
+	{
+		WorldLinAccel = WorldLinAccel.GetClampedToMaxSize(Settings.MaxLinearAcceleration);
+	}
+	if (Settings.bClampAngularAcceleration)
+	{
+		WorldAngAccel = WorldAngAccel.GetClampedToMaxSize(Settings.MaxAngularAcceleration);
+	}
+
+	// Add external velocities
+	WorldLinVel += Settings.ExternalLinearVelocity;
+	WorldAngVel += FMath::DegreesToRadians(Settings.ExternalAngularVelocity);
+
+	// Add turbulence
+	if (!Settings.ExternalTurbulenceVelocity.IsNearlyZero())
+	{
+		FVector T(
+			FPerlinHelpers::PerlinNoise1D(AbsoluteTime),
+			FPerlinHelpers::PerlinNoise1D(AbsoluteTime + 10.0),
+			FPerlinHelpers::PerlinNoise1D(AbsoluteTime + 20.0));
+		WorldLinVel += T * Settings.ExternalTurbulenceVelocity;
+	}
+
+	// Transform from world space into simulation space
+	RigParticleSimulation::FSimulationSpaceMotion Result;
+	Result.LinearVelocity = FSimVector(SimulationSpaceTM.InverseTransformVectorNoScale(WorldLinVel));
+	Result.AngularVelocity = FSimVector(SimulationSpaceTM.InverseTransformVectorNoScale(WorldAngVel));
+	Result.LinearAcceleration = FSimVector(SimulationSpaceTM.InverseTransformVectorNoScale(WorldLinAccel));
+	Result.AngularAcceleration = FSimVector(SimulationSpaceTM.InverseTransformVectorNoScale(WorldAngAccel));
+
+	Result.SpaceMovementAmount = Settings.SpaceMovementAmount;
+	Result.LinearDragMultiplier =
+		FMath::Clamp(Settings.LinearDragMultiplier * Settings.SpaceMovementAmount, 0.0f, 1.0f);
+	Result.AngularDragMultiplier =
+		FMath::Clamp(Settings.AngularDragMultiplier * Settings.SpaceMovementAmount, 0.0f, 1.0f);
+	Result.bAccelerationsValid = UpdatesSinceReset >= 2;
+
+	// External drag force: -ExternalLinearDrag * ExternalLinearVelocity (component-wise scaling in
+	// component space, then transformed to sim space). Applied as a mass-dependent force.
+	Result.ExternalDragForce = FSimVector::ZeroVector;
+	if (!Settings.ExternalLinearDrag.IsNearlyZero() && !Settings.ExternalLinearVelocity.IsNearlyZero())
+	{
+		// Per-axis scaling in component space (where the user authors the drag vector)
+		const FVector ExtVelCS = ComponentTM.InverseTransformVectorNoScale(Settings.ExternalLinearVelocity);
+		const FVector DragForceCS = -(Settings.ExternalLinearDrag * ExtVelCS) * Settings.SpaceMovementAmount;
+		Result.ExternalDragForce = FSimVector(
+			SimulationSpaceTM.InverseTransformVectorNoScale(ComponentTM.TransformVectorNoScale(DragForceCS)));
+	}
+
+	return Result;
+}
diff --git a/Engine/Plugins/Experimental/ControlRigDynamics/Source/ControlRigDynamics/Private/RigDynamicsSolver.cpp b/Engine/Plugins/Experimental/ControlRigDynamics/Source/ControlRigDynamics/Private/RigDynamicsSolver.cpp
new file mode 100644
--- /dev/null
+++ b/Engine/Plugins/Experimental/ControlRigDynamics/Source/ControlRigDynamics/Private/RigDynamicsSolver.cpp
@@ -0,0 +1,696 @@
+// Copyright Epic Games, Inc. All Rights Reserved.
+
+#include "RigDynamicsSolver.h"
+
+#include "RigParticleSimulation.h"
+
+#include "Rigs/RigHierarchy.h"
+#include "Rigs/RigHierarchyController.h"
+#include "Units/RigUnitContext.h"
+#include "RigVMCore/RigVMExecuteContext.h"
+
+//======================================================================================================================
+FRigDynamicsSolver::FRigDynamicsSolver(const FName InOwnerName) : OwnerName(InOwnerName)
+{
+}
+
+//======================================================================================================================
+int32 FRigDynamicsSolver::GetParticleIndexSlow(const FRigComponentKey& ComponentKey) const
+{
+	for (int32 Index = 0; Index < ParticleOwnerComponents.Num(); ++Index)
+	{
+		if (ParticleOwnerComponents[Index].GetComponentKey() == ComponentKey)
+		{
+			return Index;
+		}
+	}
+	return INDEX_NONE;
+}
+
+//======================================================================================================================
+int32 FRigDynamicsSolver::GetParticleIndexSlow(const FRigBaseElement& Element, const URigHierarchy& Hierarchy)
+{
+	for (int32 Index = 0; Index < ParticleOwnerComponents.Num(); ++Index)
+	{
+		if (ParticleOwnerComponents[Index].GetElement(&Hierarchy) == &Element)
+		{
+			return Index;
+		}
+	}
+	return INDEX_NONE;
+}
+
+//======================================================================================================================
+int32 FRigDynamicsSolver::GetParentParticleIndexSlow(
+	const FRigBaseElement& Element, const URigHierarchy& Hierarchy)
+{
+	const FRigBaseElement* ParentElement = Hierarchy.GetFirstParent(&Element);
+	while (ParentElement)
+	{
+		int32 Index = GetParticleIndexSlow(*ParentElement, Hierarchy);
+		if (Index != INDEX_NONE)
+		{
+			return Index;
+		}
+		ParentElement = Hierarchy.GetFirstParent(ParentElement);
+	}
+	return INDEX_NONE;
+}
+
+//======================================================================================================================
+int32 FRigDynamicsSolver::GetParentParticleIndexSlow(int32 ParticleIndex, const URigHierarchy& Hierarchy)
+{
+	if (ParticleOwnerComponents.IsValidIndex(ParticleIndex))
+	{
+		if (const FRigBaseElement* Element = ParticleOwnerComponents[ParticleIndex].GetElement(&Hierarchy))
+		{
+			return GetParentParticleIndexSlow(*Element, Hierarchy);
+		}
+	}
+	return INDEX_NONE;
+}
+
+//======================================================================================================================
+void FRigDynamicsSolver::Instantiate(
+	const FRigVMExecuteContext&        ExecuteContext,
+	URigHierarchy&                     Hierarchy,
+	const FRigDynamicsSolverComponent& SolverComponent)
+{
+	QUICK_SCOPE_CYCLE_COUNTER(STAT_RigDynamicsSolver_Instantiate);
+
+	if (!bNeedToInstantiate)
+	{
+		return;
+	}
+
+	// Build a simulation space state with just the transforms set (DeltaTime=0 skips velocity
+	// calculation). This is needed for coordinate conversions during instantiation.
+	FRigDynamicsSimulationSpaceState SpaceState;
+	UpdateSimulationSpaceStateTransforms(ExecuteContext, Hierarchy, SolverComponent, 0.0f, SpaceState);
+
+	ParticleTargetTMs.Reset();
+	ParticleOwnerComponents.Reset();
+	HardDistanceConstraintOwnerComponents.Reset();
+	SoftDistanceConstraintOwnerComponents.Reset();
+	ColliderOwnerComponents.Reset();
+
+	SimulationState = RigParticleSimulation::FSimulationState();
+
+	// Colliders. Add these first so that they're there ready to set no-collision when we add the particles
+	for (const FRigComponentKey& ColliderComponentKey : SolverComponent.Colliders)
+	{
+		if (const FRigDynamicsColliderComponent* ColliderComponent = GetCollider(Hierarchy, ColliderComponentKey))
+		{
+			RigParticleSimulation::FCollider& Collider = SimulationState.Colliders.AddDefaulted_GetRef();
+			FCachedRigComponent& CachedOwnerComponent = ColliderOwnerComponents.Add_GetRef(
+				FCachedRigComponent(ColliderComponentKey, &Hierarchy, true));
+
+			Collider.TargetTM = SpaceState.ConvertComponentSpaceTransformToSimSpace(
+				GetGlobalTransform(Hierarchy, CachedOwnerComponent));
+			Collider.PreviousTM = Collider.TargetTM;
+
+			for (const FRigDynamicsColliderBox& Box : ColliderComponent->Colliders.Boxes)
+			{
+				RigParticleSimulation::FBoxCollider& BoxCollider = Collider.BoxColliders.AddDefaulted_GetRef();
+				BoxCollider.TM = Box.TM;
+				BoxCollider.Extents = FSimVector(Box.Extents);
+			}
+			for (const FRigDynamicsColliderCapsule& Capsule : ColliderComponent->Colliders.Capsules)
+			{
+				RigParticleSimulation::FCapsuleCollider& CapsuleCollider =
+					Collider.CapsuleColliders.AddDefaulted_GetRef();
+				CapsuleCollider.TM = Capsule.TM;
+				CapsuleCollider.Length = Capsule.Length;
+				CapsuleCollider.Radius = Capsule.Radius;
+			}
+			for (const FRigDynamicsColliderPlane& Plane : ColliderComponent->Colliders.Planes)
+			{
+				RigParticleSimulation::FPlaneCollider& PlaneCollider = 
+					Collider.PlaneColliders.AddDefaulted_GetRef();
+				PlaneCollider.TM = Plane.TM;
+				PlaneCollider.Extents = FSimVector2D(Plane.Extents);
+			}
+		}
+	}
+	// All the colliders have been added
+	const int32 NumColliders = SimulationState.Colliders.Num();
+
+	// Set up our records based on the solver component data
+	for (const FRigComponentKey& ParticleComponentKey : SolverComponent.Particles)
+	{
+		// Set up the particles. 
+		if (const FRigDynamicsParticleComponent* ParticleComponent = GetParticle(Hierarchy, ParticleComponentKey))
+		{
+			RigParticleSimulation::FParticleInfo& ParticleInfo = SimulationState.ParticleInfos.AddDefaulted_GetRef();
+			RigParticleSimulation::FParticle& Particle = SimulationState.Particles.AddDefaulted_GetRef();
+			RigParticleSimulation::FParticleNoCollision& ParticleNoCollision =
+				SimulationState.ParticleNoCollision.AddDefaulted_GetRef();
+			ParticleNoCollision.NoCollisionColliderIndices.Init(false, NumColliders);
+			RigParticleSimulation::FParticleToParticleCollision& ParticleToParticleCollision =
+				SimulationState.ParticleToParticleCollision.AddDefaulted_GetRef();
+			SimulationState.ParticleTargets.AddDefaulted();
+			SimulationState.ParticleColliders.AddDefaulted();
+
+			ParticleOwnerComponents.Add(FCachedRigComponent(ParticleComponentKey, &Hierarchy, true));
+			ParticleTargetTMs.Add(GetGlobalTransform(Hierarchy, *ParticleComponent));
+
+			ParticleInfo.TargetPosition = FSimVector(
+				SpaceState.ConvertComponentSpacePositionToSimSpace(GetGlobalLocation(Hierarchy, *ParticleComponent)));
+			ParticleInfo.PreviousTargetPosition = ParticleInfo.TargetPosition;
+			ParticleInfo.GravityMultiplier = ParticleComponent->ParticleProperties.GravityMultiplier;
+			ParticleInfo.MovementType = ParticleComponent->ParticleProperties.MovementType;
+
+			Particle.Position = ParticleInfo.TargetPosition;
+			Particle.PrevPosition = ParticleInfo.TargetPosition;
+
+			// Set inverse mass here as it can be overridden if the particle needs to be kinematic
+			Particle.InvMass =
+				ParticleComponent->ParticleProperties.MovementType == ERigParticleSimulationMovementType::Simulated
+				? 1.0f / FMath::Max(ParticleComponent->ParticleProperties.Mass, KINDA_SMALL_NUMBER)
+				: 0.0f;
+
+			// Target and collision data etc will be set during the updates, so don't do it here
+
+			// Add the no-collision between particles and colliders
+			for (const FRigComponentKey& NoCollisionColliderKey : 
+				ParticleComponent->ParticleProperties.NoCollisionColliders)
+			{
+				for (int32 ColliderIndex = 0 ; ColliderIndex != ColliderOwnerComponents.Num() ; ++ColliderIndex)
+				{
+					if (ColliderOwnerComponents[ColliderIndex].GetComponentKey() == NoCollisionColliderKey)
+					{
+						ParticleNoCollision.NoCollisionColliderIndices[ColliderIndex] = true;
+					}
+				}
+			}
+
+			// Particle-to-particle collision is resolved in a second pass below
+		}
+	}
+
+	// All the particles have been added
+	const int32 NumParticles = SimulationState.Particles.Num();
+
+	// Store the parent particle indices - these will be needed for targets and joint limits
+	// TODO implement angular (cone) limits, which will need the ParentParentParticleIndex
+	for (int32 ParticleIndex = 0; ParticleIndex != NumParticles; ++ParticleIndex)
+	{
+		RigParticleSimulation::FParticleInfo& ParticleInfo = SimulationState.ParticleInfos[ParticleIndex];
+		ParticleInfo.ParentParticleIndex = GetParentParticleIndexSlow(ParticleIndex, Hierarchy);
+
+		if (ParticleInfo.ParentParticleIndex == INDEX_NONE)
+		{
+			// We didn't find a parent - make sure the child is kinematic, as this will be the
+			// root-most particle in a chain, and we don't want it to fall off.
+			SimulationState.Particles[ParticleIndex].InvMass = 0.0f;
+			ParticleInfo.MovementType = ERigParticleSimulationMovementType::Kinematic;
+		}
+		else
+		{
+			RigParticleSimulation::FParticleInfo& ParentParticleInfo = 
+				SimulationState.ParticleInfos[ParticleInfo.ParentParticleIndex];
+			++ParentParticleInfo.NumChildren;
+
+			// Check if the parent particle's bone is the direct skeleton
+			// parent (no skipped bones). When true, we can avoid a
+			// GetGlobalTransform on the child in UpdatePostDynamics.
+			const FRigBaseElement* ChildElement = ParticleOwnerComponents[ParticleIndex].GetElement(&Hierarchy);
+			const FRigBaseElement* DirectParent = Hierarchy.GetFirstParent(ChildElement);
+			const FRigBaseElement* ParentParticleElement = 
+				ParticleOwnerComponents[ParticleInfo.ParentParticleIndex].GetElement(&Hierarchy);
+			ParticleInfo.bParentParticleIsDirectParent = (DirectParent == ParentParticleElement);
+		}
+	}
+
+	// Sort particles root-to-leaf. Skeletal constraints created below will inherit this order.
+	const TArray<int32> SortedOrder = RigParticleSimulation::SortParticlesRootToLeaf(SimulationState);
+	{
+		TArray<FTransform> TempTMs;              TempTMs.SetNum(NumParticles);
+		TArray<FCachedRigComponent> TempOwners;  TempOwners.SetNum(NumParticles);
+		for (int32 NewIndex = 0; NewIndex < NumParticles; ++NewIndex)
+		{
+			TempTMs[NewIndex] = MoveTemp(ParticleTargetTMs[SortedOrder[NewIndex]]);
+			TempOwners[NewIndex] = MoveTemp(ParticleOwnerComponents[SortedOrder[NewIndex]]);
+		}
+		ParticleTargetTMs = MoveTemp(TempTMs);
+		ParticleOwnerComponents = MoveTemp(TempOwners);
+	}
+
+	for (int32 ParticleIndex = 0; ParticleIndex != NumParticles; ++ParticleIndex)
+	{
+		// Automatically disable collision between particles and colliders, if they are on the
+		// same element (which is generally a bad idea, as colliders don't track the
+		// live/simulated pose)
+		RigParticleSimulation::FParticleNoCollision& ParticleNoCollision = 
+			SimulationState.ParticleNoCollision[ParticleIndex];
+
+		// Don't collide with self
+		{
+			const FRigElementKey& OwningElementKey = ParticleOwnerComponents[ParticleIndex].GetElementKey();
+			for (int32 ColliderIndex = 0 ; ColliderIndex != NumColliders ; ++ColliderIndex)
+			{
+				const FRigElementKey& ColliderElementKey = ColliderOwnerComponents[ColliderIndex].GetElementKey();
+				if (OwningElementKey == ColliderElementKey)
+				{
+					ParticleNoCollision.NoCollisionColliderIndices[ColliderIndex] = true;
+				}
+			}
+		}
+
+		// And don't collide with the parent collider either
+		const RigParticleSimulation::FParticleInfo& ParticleInfo = SimulationState.ParticleInfos[ParticleIndex];
+		int32 ParentParticleIndex = ParticleInfo.ParentParticleIndex;
+		if (ParentParticleIndex != INDEX_NONE)
+		{
+			const FRigElementKey& OwningElementKey = ParticleOwnerComponents[ParentParticleIndex].GetElementKey();
+			for (int32 ColliderIndex = 0 ; ColliderIndex != NumColliders ; ++ColliderIndex)
+			{
+				const FRigElementKey& ColliderElementKey = ColliderOwnerComponents[ColliderIndex].GetElementKey();
+				if (OwningElementKey == ColliderElementKey)
+				{
+					ParticleNoCollision.NoCollisionColliderIndices[ColliderIndex] = true;
+				}
+			}
+		}
+	}
+
+	// Resolve particle-to-particle collision pairs. This is done as a second pass so that all
+	// particles have been added and can be found by key. Each unique pair is stored only on the
+	// lower-indexed particle to avoid solving the same pair twice per iteration. Note that the
+	// solve will project both particles at the same time, so the interaction only needs to be
+	// stored on one of them.
+	for (int32 ParticleIndex = 0; ParticleIndex != NumParticles; ++ParticleIndex)
+	{
+		if (const FRigDynamicsParticleComponent* PC = GetParticle(Hierarchy, ParticleOwnerComponents[ParticleIndex]))
+		{
+			for (const FRigComponentKey& CollisionParticleKey : PC->ParticleProperties.CollisionParticles)
+			{
+				const int32 OtherIndex = GetParticleIndexSlow(CollisionParticleKey);
+				if (OtherIndex != INDEX_NONE && OtherIndex != ParticleIndex)
+				{
+					const int32 LowerIndex = FMath::Min(ParticleIndex, OtherIndex);
+					const int32 UpperIndex = FMath::Max(ParticleIndex, OtherIndex);
+					SimulationState.ParticleToParticleCollision[LowerIndex]
+						.CollisionParticleIndices.AddUnique(UpperIndex);
+				}
+			}
+		}
+	}
+
+	// Now the distance constraints, between each particle and a parent (that might not be the
+	// immediate parent).
+	for (int32 ParticleIndex = 0; ParticleIndex != NumParticles; ++ParticleIndex)
+	{
+		const int32 ParentIndex = SimulationState.ParticleInfos[ParticleIndex].ParentParticleIndex;
+
+		if (ParentIndex != INDEX_NONE)
+		{
+			// Make a constraint. 
+			RigParticleSimulation::FHardDistanceConstraintInfo& ConstraintInfo = 
+				SimulationState.SkeletalConstraintInfos.AddDefaulted_GetRef();
+			RigParticleSimulation::FHardDistanceConstraint& Constraint = 
+				SimulationState.SkeletalConstraints.AddDefaulted_GetRef();
+			// Note that we don't actually need to know the owner (of the child particle), Even if
+			// we want to update based on the bone lengths changing, we can get that from the
+			// particle positions.
+
+			// Leave the constraint component invalid in ConstraintInfo as we don't have one
+			ConstraintInfo.ParentIndex = ParentIndex;
+			ConstraintInfo.ChildIndex = ParticleIndex;
+
+			Constraint.TargetDistance = FSimVector::Dist(
+				FSimVector(GetGlobalTransform(Hierarchy, ParticleOwnerComponents[ParticleIndex]).GetLocation()),
+				FSimVector(GetGlobalTransform(Hierarchy, ParticleOwnerComponents[ParentIndex]).GetLocation()));
+		}
+	}
+
+	// User-defined constraints
+	for (const FRigComponentKey& ConstraintComponentKey : SolverComponent.Constraints)
+	{
+		if (const FRigDynamicsConstraintComponent* ConstraintComponent = 
+			GetConstraint(Hierarchy, ConstraintComponentKey))
+		{
+			int32 ParentIndex = GetParticleIndexSlow(ConstraintComponent->ParentComponentKey);
+			int32 ChildIndex = GetParticleIndexSlow(ConstraintComponent->ChildComponentKey);
+
+			if (ParentIndex == INDEX_NONE || ChildIndex == INDEX_NONE)
+			{
+				continue;
+			}
+
+			switch (ConstraintComponent->ConstraintType)
+			{
+			case ERigParticleSimulationConstraintType::Hard:
+			{
+				RigParticleSimulation::FHardDistanceConstraintInfo& ConstraintInfo = 
+					SimulationState.HardDistanceConstraintInfos.AddDefaulted_GetRef();
+				SimulationState.HardDistanceConstraints.AddDefaulted();
+				HardDistanceConstraintOwnerComponents.Add(FCachedRigComponent(ConstraintComponentKey, &Hierarchy, true));
+
+				ConstraintInfo.ParentIndex = ParentIndex;
+				ConstraintInfo.ChildIndex = ChildIndex;
+				// The rest will be set during the update
+			}
+			break;
+			case ERigParticleSimulationConstraintType::Soft:
+			{
+				RigParticleSimulation::FSoftDistanceConstraintInfo& ConstraintInfo = 
+					SimulationState.SoftDistanceConstraintInfos.AddDefaulted_GetRef();
+				SimulationState.SoftDistanceConstraints.AddDefaulted();
+				SoftDistanceConstraintOwnerComponents.Add(FCachedRigComponent(ConstraintComponentKey, &Hierarchy, true));
+
+				ConstraintInfo.ParentIndex = ParentIndex;
+				ConstraintInfo.ChildIndex = ChildIndex;
+				// The rest will be set during the update
+			}
+			break;
+			default:
+				// Unknown constraint type
+				break;
+			}
+		}
+	}
+
+	bNeedToInstantiate = false;
+}
+
+//======================================================================================================================
+void FRigDynamicsSolver::UpdatePreDynamics(URigHierarchy& Hierarchy, float DeltaTime)
+{
+	QUICK_SCOPE_CYCLE_COUNTER(STAT_RigDynamics_UpdatePreDynamics);
+
+	float InvDeltaTime = DeltaTime > KINDA_SMALL_NUMBER ? 1.0f / DeltaTime : 0.0f;
+
+	{
+		QUICK_SCOPE_CYCLE_COUNTER(STAT_RigDynamics_UpdatePreDynamics_Particles);
+
+		check(SimulationState.ParticleInfos.Num() == SimulationState.Particles.Num());
+		check(ParticleOwnerComponents.Num() == SimulationState.Particles.Num());
+		for (int32 Index = 0; Index != SimulationState.Particles.Num(); ++Index)
+		{
+			if (const FRigDynamicsParticleComponent* PC = GetParticle(Hierarchy, ParticleOwnerComponents[Index]))
+			{
+				RigParticleSimulation::UpdateParticle(
+					SimulationState, Index, PC->ParticleProperties.Mass, 
+					PC->ParticleProperties.GravityMultiplier, PC->ParticleProperties.Radius);
+
+				const FTransform TargetTM = GetGlobalTransform(Hierarchy, *PC);
+				ParticleTargetTMs[Index] = TargetTM;
+				const FVector TargetSimSpace =
+					SimulationSpaceState.ConvertComponentSpacePositionToSimSpace(TargetTM.GetLocation());
+				RigParticleSimulation::UpdateParticleTarget(
+					SimulationState, Index, InvDeltaTime, FSimVector(TargetSimSpace),
+					PC->ParticleProperties.Strength, PC->ParticleProperties.DampingRatio, 
+					PC->ParticleProperties.ExtraDamping, PC->ParticleProperties.TargetVelocityInfluence, 
+					PC->ParticleProperties.TargetMode);
+			}
+		}
+	}
+
+	{
+		QUICK_SCOPE_CYCLE_COUNTER(STAT_RigDynamics_UpdatePreDynamics_Constraints);
+
+		// We call RigParticleSimulation::UpdateSkeletalConstraint here to support bone lengths
+		// changing. This should be optional.
+		for (int32 Index = 0; Index != SimulationState.SkeletalConstraints.Num(); ++Index)
+		{
+			RigParticleSimulation::UpdateSkeletalConstraint(SimulationState, Index);
+		}
+
+		// Update user-specified hard constraints
+		check(SimulationState.HardDistanceConstraints.Num() == SimulationState.HardDistanceConstraintInfos.Num());
+		check(HardDistanceConstraintOwnerComponents.Num() == SimulationState.HardDistanceConstraintInfos.Num());
+		for (int32 Index = 0; Index != SimulationState.HardDistanceConstraints.Num(); ++Index)
+		{
+			if (const FRigDynamicsConstraintComponent* CC = GetConstraint(
+				Hierarchy, HardDistanceConstraintOwnerComponents[Index]))
+			{
+				RigParticleSimulation::UpdateHardDistanceConstraint(
+					SimulationState, Index, CC->LengthMultiplier, CC->ExtraLength);
+			}
+		}
+
+		// Update user-specified soft constraints
+		check(SimulationState.SoftDistanceConstraints.Num() == SimulationState.SoftDistanceConstraintInfos.Num());
+		check(SoftDistanceConstraintOwnerComponents.Num() == SimulationState.SoftDistanceConstraintInfos.Num());
+		for (int32 Index = 0; Index != SimulationState.SoftDistanceConstraints.Num(); ++Index)
+		{
+			if (const FRigDynamicsConstraintComponent* CC = GetConstraint(
+				Hierarchy, SoftDistanceConstraintOwnerComponents[Index]))
+			{
+				RigParticleSimulation::UpdateSoftDistanceConstraint(SimulationState, Index,
+					CC->LengthMultiplier, CC->ExtraLength, CC->Strength, CC->DampingRatio, CC->ExtraDamping);
+			}
+		}
+	}
+
+	{
+		QUICK_SCOPE_CYCLE_COUNTER(STAT_RigDynamics_UpdatePreDynamics_Colliders);
+
+		// Update colliders. Shape arrays are expected to match what was set up in Instantiate (dynamic
+		// shape array changes are not yet supported, only size/position/orientation changes).
+		check(ColliderOwnerComponents.Num() == SimulationState.Colliders.Num());
+		for (int32 Index = 0; Index != SimulationState.Colliders.Num(); ++Index)
+		{
+			if (const FRigDynamicsColliderComponent* ColliderComponent = GetCollider(
+				Hierarchy, ColliderOwnerComponents[Index]))
+			{
+				const FTransform ColliderSimSpace = SimulationSpaceState.ConvertComponentSpaceTransformToSimSpace(
+					GetGlobalTransform(Hierarchy, *ColliderComponent));
+				RigParticleSimulation::UpdateCollider(SimulationState, Index, ColliderSimSpace);
+
+				RigParticleSimulation::FCollider& Collider = SimulationState.Colliders[Index];
+
+				// Now update the individual shapes (optional - only needed if things change)
+				for (int32 BoxIndex = 0; BoxIndex != ColliderComponent->Colliders.Boxes.Num(); ++BoxIndex)
+				{
+					if (ensure(Collider.BoxColliders.IsValidIndex(BoxIndex)))
+					{
+						const FRigDynamicsColliderBox& Box = ColliderComponent->Colliders.Boxes[BoxIndex];
+						Collider.BoxColliders[BoxIndex].Update(FSimTransform(Box.TM), FSimVector(Box.Extents));
+					}
+				}
+				for (int32 CapsuleIndex = 0; CapsuleIndex != ColliderComponent->Colliders.Capsules.Num(); ++CapsuleIndex)
+				{
+					if (ensure(Collider.CapsuleColliders.IsValidIndex(CapsuleIndex)))
+					{
+						const FRigDynamicsColliderCapsule& Capsule = ColliderComponent->Colliders.Capsules[CapsuleIndex];
+						Collider.CapsuleColliders[CapsuleIndex].Update(
+							FSimTransform(Capsule.TM), Capsule.Length, Capsule.Radius);
+					}
+				}
+				for (int32 PlaneIndex = 0; PlaneIndex != ColliderComponent->Colliders.Planes.Num(); ++PlaneIndex)
+				{
+					if (ensure(Collider.PlaneColliders.IsValidIndex(PlaneIndex)))
+					{
+						const FRigDynamicsColliderPlane& Plane = ColliderComponent->Colliders.Planes[PlaneIndex];
+						Collider.PlaneColliders[PlaneIndex].Update(
+							FSimTransform(Plane.TM), FSimVector2D(Plane.Extents));
+					}
+				}
+			}
+		}
+	}
+}
+
+//======================================================================================================================
+void FRigDynamicsSolver::UpdatePostDynamics(
+	URigHierarchy& Hierarchy, const FRigDynamicsSolverComponent& SolverComponent, float Alpha)
+{
+	QUICK_SCOPE_CYCLE_COUNTER(STAT_RigDynamics_UpdatePostDynamics);
+
+	// Particles are in depth order (root-to-leaf), ensured by SortParticlesRootToLeaf in Instantiate.
+	// Note: branching chains where a parent is shared by multiple children could rotate the parent
+	// twice. The multi-child case below handles this by only setting position, not orientation.
+	for (int32 Index = 0; Index < SimulationState.Particles.Num(); ++Index)
+	{
+		// Clear forces ready for the next update (before any early
+		// continues so it always runs).
+		SimulationState.ParticleInfos[Index].ExternalForce = FSimVector::ZeroVector;
+
+		const RigParticleSimulation::FParticle& Particle = SimulationState.Particles[Index];
+		const RigParticleSimulation::FParticleInfo& ParticleInfo = SimulationState.ParticleInfos[Index];
+
+		if (ParticleInfo.NumChildren > 1)
+		{
+			// When there are multiple children, we always:
+			// 1. track the particle position exactly
+			// 2. Maintain the original component-space orientation.
+			const int32 BoneIndex = GetElementIndex(Hierarchy, ParticleOwnerComponents[Index]);
+			FTransform OriginalTM = Hierarchy.GetGlobalTransform(BoneIndex);
+			const FVector ParticlePosComponentSpace =
+				SimulationSpaceState.ConvertSimSpacePositionToComponentSpace(FVector(Particle.Position));
+			FTransform ParticleTM(ParticleTargetTMs[Index].GetRotation(), ParticlePosComponentSpace);
+			FTransform BlendedTM;
+			BlendedTM.Blend(OriginalTM, ParticleTM, Alpha);
+			Hierarchy.SetGlobalTransform(BoneIndex, BlendedTM);
+		}
+		else
+		{
+			bool bForceSettingPosition = false;
+			if (SolverComponent.Settings.bSetBoneOrientations)
+			{
+				const int32 ParentIndex = ParticleInfo.ParentParticleIndex;
+
+				// Rotate parent bone so it points toward this particle
+				if (ParentIndex != INDEX_NONE)
+				{
+					const RigParticleSimulation::FParticleInfo& ParentParticleInfo = 
+						SimulationState.ParticleInfos[ParentIndex];
+
+					if (Particle.InvMass == 0.0f && SimulationState.Particles[ParentIndex].InvMass == 0.0f)
+					{
+						// If both this particle and its parent are kinematic, neither moved during
+						// simulation, so the parent bone's orientation is already correct from animation.
+					}
+					else
+					{
+						// Only rotate joints when there is a single child. Alternatively, we could just
+						// apply a fraction of the delta rotation, but this would be tricky as we want to
+						// update incrementally.
+						if (ParentParticleInfo.NumChildren != 1)
+						{
+							bForceSettingPosition = true;
+						}
+						else
+						{
+							const FVector ChildParticlePos = SimulationSpaceState.ConvertSimSpacePositionToComponentSpace(
+								FVector(Particle.Position));
+							const FVector ParentParticlePos = SimulationSpaceState.ConvertSimSpacePositionToComponentSpace(
+								FVector(SimulationState.Particles[ParentIndex].Position));
+							const FVector DirectionFromParticles = ChildParticlePos - ParentParticlePos;
+
+							const int32 ParentBoneIndex = GetElementIndex(Hierarchy, ParticleOwnerComponents[ParentIndex]);
+							const int32 ChildBoneIndex = GetElementIndex(Hierarchy, ParticleOwnerComponents[Index]);
+
+							FTransform ParentBoneTransform =   Hierarchy.GetGlobalTransform(ParentBoneIndex);
+
+							// When the parent particle's bone is the direct skeleton parent we can derive the
+							// child's global position from the (already known) parent transform + the child's
+							// local transform, avoiding a GetGlobalTransform call that would recompute dirty globals.
+							const FVector ChildBonePos = ParticleInfo.bParentParticleIsDirectParent
+								? ParentBoneTransform.TransformPositionNoScale(
+									Hierarchy.GetLocalTransform(ChildBoneIndex).GetLocation())
+								: FVector(Hierarchy.GetGlobalTransform(ChildBoneIndex).GetLocation());
+
+							const FVector DirectionFromBones = ChildBonePos - ParentBoneTransform.GetLocation();
+
+							const FQuat DeltaRotation = FQuat::FindBetweenVectors(
+								DirectionFromBones, DirectionFromParticles);
+							const FQuat BlendedRotation = FQuat::Slerp(FQuat::Identity, DeltaRotation, Alpha);
+							ParentBoneTransform.SetRotation(BlendedRotation * ParentBoneTransform.GetRotation());
+							Hierarchy.SetGlobalTransform(ParentBoneIndex, ParentBoneTransform);
+						}
+					}
+				}
+			}
+
+			// Position fixup
+			if (SolverComponent.Settings.bSetBonePositions || bForceSettingPosition)
+			{
+				const int32 BoneIndex = GetElementIndex(Hierarchy, ParticleOwnerComponents[Index]);
+				FTransform Transform = Hierarchy.GetGlobalTransform(BoneIndex);
+				const FVector ParticlePosCS = SimulationSpaceState.ConvertSimSpacePositionToComponentSpace(
+					FVector(Particle.Position));
+				const FVector BlendedPosition = FMath::Lerp(Transform.GetLocation(), ParticlePosCS, Alpha);
+				Transform.SetLocation(BlendedPosition);
+				Hierarchy.SetGlobalTransform(BoneIndex, Transform);
+			}
+		}
+	}
+}
+
+//======================================================================================================================
+void FRigDynamicsSolver::UpdateSimulationSpaceStateTransforms(
+	const FRigVMExecuteContext&        ExecuteContext,
+	URigHierarchy&                     Hierarchy,
+	const FRigDynamicsSolverComponent& SolverComponent,
+	const float                        DeltaTime,
+	FRigDynamicsSimulationSpaceState&  OutState)
+{
+	const FRigDynamicsSolverSettings& SolverSettings = SolverComponent.Settings;
+
+	FTransform ComponentTM = ExecuteContext.GetToWorldSpaceTransform();
+
+	FTransform SimulationSpaceTM;
+	switch (SolverSettings.SimulationSpace)
+	{
+	case ERigDynamicsSimulationSpace::World:
+	{
+		// Identity: simulation space IS world space
+		break;
+	}
+	case ERigDynamicsSimulationSpace::Component:
+	{
+		SimulationSpaceTM = ComponentTM;
+		break;
+	}
+	case ERigDynamicsSimulationSpace::SpaceBone:
+	{
+		if (SolverSettings.SpaceBone.IsValid())
+		{
+			SimulationSpaceTM = Hierarchy.GetGlobalTransform(SolverSettings.SpaceBone) * ComponentTM;
+		}
+		else
+		{
+			SimulationSpaceTM = ComponentTM;
+		}
+		break;
+	}
+	}
+
+	OutState.Update(ComponentTM, SimulationSpaceTM, DeltaTime);
+}
+
+//======================================================================================================================
+// Note that OwningActorPtr will be needed for things like simulation space and collision
+void FRigDynamicsSolver::StepSimulation(
+	const FRigVMExecuteContext&        ExecuteContext,
+	URigHierarchy&                     Hierarchy,
+	const FRigDynamicsSolverComponent& SolverComponent,
+	const AActor*                      OwningActorPtr,
+	const float                        DeltaTimeOverride,
+	const float                        Alpha)
+{
+	QUICK_SCOPE_CYCLE_COUNTER(STAT_RigDynamics_StepSimulation);
+
+	// Ensure we're instantiated
+	if (bNeedToInstantiate)
+	{
+		Instantiate(ExecuteContext, Hierarchy, SolverComponent);
+	}
+
+	float DeltaTime = ExecuteContext.GetDeltaTime();
+	if (DeltaTimeOverride > 0.0f)
+	{
+		DeltaTime = DeltaTimeOverride;
+	}
+	else if (DeltaTimeOverride < 0.0f)
+	{
+		DeltaTime = 0.0f;
+	}
+
+	const FRigDynamicsSolverSettings& SolverSettings = SolverComponent.Settings;
+	const FRigDynamicsSimulationSpaceSettings& SimulationSpaceSettings = SolverComponent.SimulationSpaceSettings;
+
+	UpdateSimulationSpaceStateTransforms(ExecuteContext, Hierarchy, SolverComponent, DeltaTime, SimulationSpaceState);
+
+	SimulationSpaceMotion = SimulationSpaceState.CalculateMotion(
+		SimulationSpaceSettings, ExecuteContext.GetAbsoluteTime());
+
+	RigParticleSimulation::FSolverSettings ParticleSolverSettings;
+	ParticleSolverSettings.Gravity = SimulationSpaceState.ConvertWorldVectorToSimSpace(SolverSettings.Gravity);
+	ParticleSolverSettings.NumIterations = SolverSettings.NumIterations;
+	ParticleSolverSettings.NumConstraintSubIterations = SolverSettings.NumConstraintSubIterations;
+	ParticleSolverSettings.MaxTimeStep = SolverSettings.MaxTimeStep;
+	ParticleSolverSettings.MaxNumSteps = SolverSettings.MaxNumSteps;
+
+	if (DeltaTime > 0.0f)
+	{
+		UpdatePreDynamics(Hierarchy, DeltaTime);
+
+		RigParticleSimulation::Simulate(SimulationState, ParticleSolverSettings, SimulationSpaceMotion, DeltaTime);
+	}
+
+	UpdatePostDynamics(Hierarchy, SolverComponent, Alpha);
+}
diff --git a/Engine/Plugins/Experimental/ControlRigDynamics/Source/ControlRigDynamics/Private/RigDynamicsSolver.h b/Engine/Plugins/Experimental/ControlRigDynamics/Source/ControlRigDynamics/Private/RigDynamicsSolver.h
new file mode 100644
--- /dev/null
+++ b/Engine/Plugins/Experimental/ControlRigDynamics/Source/ControlRigDynamics/Private/RigDynamicsSolver.h
@@ -0,0 +1,104 @@
+// Copyright Epic Games, Inc. All Rights Reserved.
+
+#pragma once
+
+#include "RigDynamicsHelpers.h"
+
+#include "RigParticleSimulation.h"
+#include "RigDynamicsSimulationSpace.h"
+
+#include "Rigs/RigHierarchyCache.h"
+
+using RigParticleSimulation::FSimVector;
+using RigParticleSimulation::FSimTransform;
+using RigParticleSimulation::FSimVector2D;
+
+struct FRigVMExecuteContext;
+struct FRigVMDrawInterface;
+class URigHierarchy;
+
+//======================================================================================================================
+// This is the main solver for Dynamics, handling the particles and constraints, as well as reading
+// from/writing back to the hierarchy.
+//======================================================================================================================
+struct FRigDynamicsSolver
+{
+public:
+	FRigDynamicsSolver(const FName InOwnerName);
+
+	// This will initialise/create the simulation and then create everything we need in it. 
+	void Instantiate(
+		const FRigVMExecuteContext&        ExecuteContext, 
+		URigHierarchy&                     Hierarchy,
+		const FRigDynamicsSolverComponent& SolverComponent);
+
+	// Integrates the simulation forwards.
+	// If DeltaTimeOverride is +ve, then that value is used.
+	// If it is zero, then delta time is taken from the execute context
+	// If it is negative, then the simulation isn't stepped.
+	void StepSimulation(
+		const FRigVMExecuteContext&        ExecuteContext, 
+		URigHierarchy&                     Hierarchy,
+		const FRigDynamicsSolverComponent& SolverComponent,
+		const AActor*                      OwningActorPtr,
+		const float                        DeltaTimeOverride, 
+		const float                        Alpha);
+
+	// Draws shapes etc 
+	void Draw(
+		FRigVMDrawInterface*                     DI,
+		const URigHierarchy&                     Hierarchy,
+		const UWorld*                            DebugWorld,
+		const FRigDynamicsVisualizationSettings& VisualizationSettings) const;
+
+private:
+	// Returns the index into ComponentInfos/Particles for the given element, or INDEX_NONE
+	int32 GetParticleIndexSlow(const FRigBaseElement& Element, const URigHierarchy& Hierarchy);
+
+	// Returns the index into ComponentInfos/Particles for the given element, or INDEX_NONE
+	int32 GetParticleIndexSlow(const FRigComponentKey& ComponentKey) const;
+
+	// Returns the index into ComponentInfos/Particles for the parent of the given element, or INDEX_NONE
+	int32 GetParentParticleIndexSlow(const FRigBaseElement& Element, const URigHierarchy& Hierarchy);
+
+	// Returns the index into ComponentInfos/Particles for the parent of the given particle, or INDEX_NONE
+	int32 GetParentParticleIndexSlow(int32 ParticleIndex, const URigHierarchy& Hierarchy);
+
+	// This will prepare things - e.g. read from the hierarchy target positions etc
+	void UpdatePreDynamics(URigHierarchy& Hierarchy, float DeltaTime);
+
+	// This will read back from the simulation into the hierarchy
+	void UpdatePostDynamics(URigHierarchy& Hierarchy, const FRigDynamicsSolverComponent& SolverComponent, float Alpha);
+
+	static void UpdateSimulationSpaceStateTransforms(
+		const FRigVMExecuteContext&          ExecuteContext,
+		URigHierarchy&                       Hierarchy,
+		const FRigDynamicsSolverComponent&   SolverComponent,
+		float                                DeltaTime,
+		FRigDynamicsSimulationSpaceState&    OutState);
+
+private:
+
+	// A name that identifies what owns us - just used for debugging/logging
+	FName OwnerName;
+
+	// Instantiation can be explicit, or triggered automatically on the StepSimulation call - track
+	// whether it is needed.
+	bool bNeedToInstantiate = true;
+
+	// Each pair of arrays (Info and the Simulation) will have the same size. Infos contain data
+	// that refers to control rig, and is used to process inputs. The RigParticleSimulation data is
+	// used by the simulation, and should be packed as tightly as possible. I have tried baking the
+	// arrays to pack them into contiguous memory, but it made no measurable improvement.
+
+	TArray<FTransform>           ParticleTargetTMs;
+	TArray<FCachedRigComponent>  ParticleOwnerComponents;
+	TArray<FCachedRigComponent>  HardDistanceConstraintOwnerComponents;
+	TArray<FCachedRigComponent>  SoftDistanceConstraintOwnerComponents;
+	TArray<FCachedRigComponent>  ColliderOwnerComponents;
+
+	RigParticleSimulation::FSimulationState SimulationState;
+
+	FRigDynamicsSimulationSpaceState SimulationSpaceState;
+	RigParticleSimulation::FSimulationSpaceMotion SimulationSpaceMotion;
+};
diff --git a/Engine/Plugins/Experimental/ControlRigDynamics/Source/ControlRigDynamics/Private/RigDynamicsSolverComponent.cpp b/Engine/Plugins/Experimental/ControlRigDynamics/Source/ControlRigDynamics/Private/RigDynamicsSolverComponent.cpp
new file mode 100644
--- /dev/null
+++ b/Engine/Plugins/Experimental/ControlRigDynamics/Source/ControlRigDynamics/Private/RigDynamicsSolverComponent.cpp
@@ -0,0 +1,48 @@
+// Copyright Epic Games, Inc. All Rights Reserved.
+
+#include "RigDynamicsSolverComponent.h"
+
+#include "RigDynamicsSolver.h"
+#include "RigDynamicsObjectVersion.h"
+
+//======================================================================================================================
+void FRigDynamicsSolverComponent::Serialize(FArchive& Ar)
+{
+	Ar << Particles;
+	Ar << Colliders;
+	if (Ar.CustomVer(FRigDynamicsObjectVersion::GUID) >= FRigDynamicsObjectVersion::Constraints)
+	{
+		Ar << Constraints;
+	}
+	Ar << Settings;
+	if (Ar.CustomVer(FRigDynamicsObjectVersion::GUID) >= FRigDynamicsObjectVersion::SimulationSpace)
+	{
+		Ar << SimulationSpaceSettings;
+	}
+}
+
+//======================================================================================================================
+void FRigDynamicsSolverComponent::Save(FArchive& Ar)
+{
+	Ar.UsingCustomVersion(FRigDynamicsObjectVersion::GUID);
+	FRigBaseComponent::Save(Ar);
+	Serialize(Ar);
+}
+
+//======================================================================================================================
+void FRigDynamicsSolverComponent::Load(FArchive& Ar)
+{
+	FRigBaseComponent::Load(Ar);
+	Serialize(Ar);
+}
+
+//======================================================================================================================
+FRigDynamicsSolver* FRigDynamicsSolverComponent::GetDynamicsSolver() const
+{
+	if (!DynamicsSolverPtr.IsValid())
+	{
+		DynamicsSolverPtr = MakeShared<FRigDynamicsSolver>(TEXT("DynamicsSolver"));
+	}
+	return DynamicsSolverPtr.Get();
+}
+
diff --git a/Engine/Plugins/Experimental/ControlRigDynamics/Source/ControlRigDynamics/Private/RigDynamicsSolverExecution.cpp b/Engine/Plugins/Experimental/ControlRigDynamics/Source/ControlRigDynamics/Private/RigDynamicsSolverExecution.cpp
new file mode 100644
--- /dev/null
+++ b/Engine/Plugins/Experimental/ControlRigDynamics/Source/ControlRigDynamics/Private/RigDynamicsSolverExecution.cpp
@@ -0,0 +1,191 @@
+// Copyright Epic Games, Inc. All Rights Reserved.
+
+#include "RigDynamicsSolverExecution.h"
+#include "RigDynamicsSolver.h"
+
+#include "RigDynamicsHelpers.h"
+
+#include "ControlRig.h"
+#include "Rigs/RigHierarchyController.h"
+#include "Units/RigUnitContext.h"
+#include "Units/Execution/RigUnit_PrepareForExecution.h"
+
+#include UE_INLINE_GENERATED_CPP_BY_NAME(RigDynamicsSolverExecution)
+
+TAutoConsoleVariable<int> CVarControlRigDynamicsVisualizationOverride(
+	TEXT("ControlRig.Dynamics.VisualizationOverride"), -1,
+	TEXT("-1 uses the visualization settings from Step Dynamics Solver, 0 forces visualization to be disabled, 1 forces it to be enabled."));
+
+
+//======================================================================================================================
+// Macros to reduce boilerplate for the Get/Set solver property nodes. Note that these can't
+// be used in the header because UHT scans raw source before C++ preprocessing, so it would
+// never see the USTRUCT/UPROPERTY/RIGVM_METHOD markers inside a macro expansion.
+//======================================================================================================================
+
+#define IMPLEMENT_SET_DYNAMICS_SOLVER(PropName)                                   \
+FRigUnit_HierarchySetDynamicsSolver##PropName##_Execute()                         \
+{                                                                                 \
+	if (!ExecuteContext.Hierarchy) { return; }                                    \
+	if (FRigDynamicsSolverComponent* Component =                                  \
+		GetSolver(*ExecuteContext.Hierarchy, DynamicsSolverComponentKey))         \
+	{                                                                             \
+		Component->PropName = PropName;                                           \
+	}                                                                             \
+}
+
+#define IMPLEMENT_GET_DYNAMICS_SOLVER(PropName)                                   \
+FRigUnit_HierarchyGetDynamicsSolver##PropName##_Execute()                         \
+{                                                                                 \
+	if (!ExecuteContext.Hierarchy) { return; }                                    \
+	if (const FRigDynamicsSolverComponent* Component =                            \
+		GetSolver(*ExecuteContext.Hierarchy, DynamicsSolverComponentKey))         \
+	{                                                                             \
+		PropName = Component->PropName;                                           \
+	}                                                                             \
+}
+
+#define IMPLEMENT_SET_DYNAMICS_SOLVER_SETTING(PropName)                           \
+FRigUnit_HierarchySetDynamicsSolver##PropName##_Execute()                         \
+{                                                                                 \
+	if (!ExecuteContext.Hierarchy) { return; }                                    \
+	if (FRigDynamicsSolverComponent* Component =                                  \
+		GetSolver(*ExecuteContext.Hierarchy, DynamicsSolverComponentKey))         \
+	{                                                                             \
+		Component->Settings.PropName = PropName;                                  \
+	}                                                                             \
+}
+
+#define IMPLEMENT_GET_DYNAMICS_SOLVER_SETTING(PropName)                           \
+FRigUnit_HierarchyGetDynamicsSolver##PropName##_Execute()                         \
+{                                                                                 \
+	if (!ExecuteContext.Hierarchy) { return; }                                    \
+	if (const FRigDynamicsSolverComponent* Component =                            \
+		GetSolver(*ExecuteContext.Hierarchy, DynamicsSolverComponentKey))         \
+	{                                                                             \
+		PropName = Component->Settings.PropName;                                  \
+	}                                                                             \
+}
+
+//======================================================================================================================
+FRigUnit_SpawnDynamicsSolver_Execute()
+{
+	if (!ExecuteContext.Hierarchy)
+	{
+		return;
+	}
+	if (ExecuteContext.GetEventName() != FRigUnit_PrepareForExecution::EventName)
+	{
+		UE_CONTROLRIG_RIGUNIT_REPORT_ERROR(TEXT("SpawnDynamicsSolver can only be used during Setup"));
+		return;
+	}
+
+	if (URigHierarchyController* Controller = ExecuteContext.Hierarchy->GetController())
+	{
+		FRigHierarchyControllerInstructionBracket InstructionBracket(
+			Controller, ExecuteContext.GetInstructionIndex());
+		DynamicsSolverComponentKey = Controller->AddComponent(
+			FRigDynamicsSolverComponent::StaticStruct(), FRigDynamicsSolverComponent::GetDefaultName(), Owner);
+		if (DynamicsSolverComponentKey.IsValid())
+		{
+			if (FRigDynamicsSolverComponent* Component =
+				GetSolver(*ExecuteContext.Hierarchy, DynamicsSolverComponentKey))
+			{
+				Component->Particles = Particles;
+				Component->Colliders = Colliders;
+				Component->Settings = Settings;
+				Component->SimulationSpaceSettings = SimulationSpaceSettings;
+			}
+		}
+	}
+}
+
+//======================================================================================================================
+FRigUnit_InstantiateDynamics_Execute()
+{
+	DECLARE_SCOPE_HIERARCHICAL_COUNTER_RIGUNIT();
+	SCOPED_NAMED_EVENT(RigDynamicsSolver_Instantiate, FColor::Yellow);
+
+	if (!ExecuteContext.Hierarchy)
+	{
+		return;
+	}
+	URigHierarchy& Hierarchy = *ExecuteContext.Hierarchy;
+
+	const FRigDynamicsSolverComponent* SolverComponent = GetSolver(Hierarchy, DynamicsSolverComponentKey);
+
+	if (SolverComponent)
+	{
+		if (FRigDynamicsSolver* Solver = SolverComponent->GetDynamicsSolver())
+		{
+			Solver->Instantiate(ExecuteContext, Hierarchy, *SolverComponent);
+		}
+	}
+}
+
+//======================================================================================================================
+FRigUnit_StepDynamicsSolver_Execute()
+{
+	DECLARE_SCOPE_HIERARCHICAL_COUNTER_RIGUNIT();
+	SCOPED_NAMED_EVENT(RigDynamicsSolver_StepSimulation, FColor::Yellow);
+
+	if (!ExecuteContext.Hierarchy)
+	{
+		return;
+	}
+	URigHierarchy& Hierarchy = *ExecuteContext.Hierarchy;
+
+	const FRigDynamicsSolverComponent* SolverComponent = GetSolver(Hierarchy, CachedSolverComponent);
+		
+	if (!SolverComponent || (SolverComponent->GetKey() != DynamicsSolverComponentKey))
+	{
+		CachedSolverComponent = FCachedRigComponent(DynamicsSolverComponentKey, &Hierarchy, true);
+		SolverComponent = GetSolver(Hierarchy, CachedSolverComponent);
+	}
+
+	if (SolverComponent)
+	{
+		if (UControlRig* ControlRig = Cast<UControlRig>(Hierarchy.GetOuter()))
+		{
+			// We only set this flag once we start executing, because that's when it matters.
+			// Control rig is responsible for clearing it when the rig is constructed.
+			ControlRig->SetContainsSimulation();
+
+			if (FRigDynamicsSolver* Solver = SolverComponent->GetDynamicsSolver())
+			{
+				const USceneComponent* OwningSceneComponent = ControlRig->GetOwningSceneComponent();
+				const AActor* OwningActorPtr = OwningSceneComponent ? OwningSceneComponent->GetOwner() : nullptr;
+
+				Solver->StepSimulation(
+					ExecuteContext, Hierarchy, *SolverComponent, OwningActorPtr, DeltaTimeOverride, Alpha);
+
+				int32 VisualizationOverride = CVarControlRigDynamicsVisualizationOverride.GetValueOnAnyThread();
+				if (VisualizationOverride == 1 || 
+					(VisualizationOverride < 0 && VisualizationSettings.bEnableVisualization))
+				{
+					Solver->Draw(
+						ExecuteContext.GetDrawInterface(), Hierarchy, ExecuteContext.GetWorld(), VisualizationSettings);
+				}
+			}
+		}
+	}
+}
+
+//======================================================================================================================
+IMPLEMENT_SET_DYNAMICS_SOLVER(Settings)
+IMPLEMENT_GET_DYNAMICS_SOLVER(Settings)
+//======================================================================================================================
+IMPLEMENT_SET_DYNAMICS_SOLVER_SETTING(Gravity)
+IMPLEMENT_GET_DYNAMICS_SOLVER_SETTING(Gravity)
+//======================================================================================================================
+IMPLEMENT_SET_DYNAMICS_SOLVER_SETTING(MaxTimeStep)
+IMPLEMENT_GET_DYNAMICS_SOLVER_SETTING(MaxTimeStep)
+//======================================================================================================================
+IMPLEMENT_SET_DYNAMICS_SOLVER_SETTING(MaxNumSteps)
+IMPLEMENT_GET_DYNAMICS_SOLVER_SETTING(MaxNumSteps)
+//======================================================================================================================
+IMPLEMENT_SET_DYNAMICS_SOLVER_SETTING(NumIterations)
+IMPLEMENT_GET_DYNAMICS_SOLVER_SETTING(NumIterations)
+//======================================================================================================================
+IMPLEMENT_SET_DYNAMICS_SOLVER_SETTING(NumConstraintSubIterations)
+IMPLEMENT_GET_DYNAMICS_SOLVER_SETTING(NumConstraintSubIterations)
diff --git a/Engine/Plugins/Experimental/ControlRigDynamics/Source/ControlRigDynamics/Private/RigDynamicsSolver_Visualization.cpp b/Engine/Plugins/Experimental/ControlRigDynamics/Source/ControlRigDynamics/Private/RigDynamicsSolver_Visualization.cpp
new file mode 100644
--- /dev/null
+++ b/Engine/Plugins/Experimental/ControlRigDynamics/Source/ControlRigDynamics/Private/RigDynamicsSolver_Visualization.cpp
@@ -0,0 +1,103 @@
+// Copyright Epic Games, Inc. All Rights Reserved.
+
+#include "RigDynamicsSolver.h"
+
+#include "Engine/Engine.h"
+
+#include "Materials/MaterialInstanceDynamic.h"
+
+#include "Math/Color.h"
+
+namespace RigDynamicsSolverDraw
+{
+	const FLinearColor SimulatedParticleColor = FColor::Yellow;
+	const FLinearColor KinematicParticleColor = FColor::Blue;
+	const FLinearColor ColliderColor = FColor::Red;
+}
+
+//======================================================================================================================
+void FRigDynamicsSolver::Draw(
+	FRigVMDrawInterface*                     DI,
+	const URigHierarchy&                     Hierarchy,
+	const UWorld*                            DebugWorld,
+	const FRigDynamicsVisualizationSettings& VisualizationSettings) const
+{
+	QUICK_SCOPE_CYCLE_COUNTER(STAT_RigDynamics_Draw);
+
+	if (!DI)
+	{
+		return;
+	}
+
+	// Draw colliders at their current positions, converting from simulation space to component space
+	if (VisualizationSettings.bShowColliders)
+	{
+		for (const RigParticleSimulation::FCollider& Collider : SimulationState.Colliders)
+		{
+			for (const RigParticleSimulation::FPlaneCollider& PlaneCollider : Collider.PlaneColliders)
+			{
+				if (GEngine && GEngine->ConstraintLimitMaterialPrismatic)
+				{
+					DI->DrawPlane(
+						SimulationSpaceState.ConvertSimSpaceTransformToComponentSpace(PlaneCollider.TM * Collider.TM),
+						FVector2D(PlaneCollider.Extents * 0.5f),
+						RigDynamicsSolverDraw::ColliderColor,
+						true,
+						RigDynamicsSolverDraw::ColliderColor,
+						GEngine->ConstraintLimitMaterialPrismatic->GetRenderProxy());
+				}
+			}
+			for (const RigParticleSimulation::FBoxCollider& BoxCollider : Collider.BoxColliders)
+			{
+				DI->DrawBox(
+					SimulationSpaceState.ConvertSimSpaceTransformToComponentSpace(BoxCollider.TM * Collider.TM),
+					FTransform(FQuat::Identity, FVector::ZeroVector, FVector(BoxCollider.Extents)),
+					RigDynamicsSolverDraw::ColliderColor,
+					VisualizationSettings.LineThickness);
+			}
+			for (const RigParticleSimulation::FCapsuleCollider& CapsuleCollider : Collider.CapsuleColliders)
+			{
+				const FTransform ShapeTM =
+					SimulationSpaceState.ConvertSimSpaceTransformToComponentSpace(CapsuleCollider.TM * Collider.TM);
+				if (CapsuleCollider.Length < KINDA_SMALL_NUMBER)
+				{
+					DI->DrawSphere(ShapeTM, FTransform(), CapsuleCollider.Radius,
+						RigDynamicsSolverDraw::ColliderColor,
+						VisualizationSettings.LineThickness,
+						VisualizationSettings.ShapeDetail);
+				}
+				else
+				{
+					DI->DrawCapsule(ShapeTM, FTransform(), CapsuleCollider.Radius, CapsuleCollider.Length,
+						RigDynamicsSolverDraw::ColliderColor,
+						VisualizationSettings.LineThickness,
+						VisualizationSettings.ShapeDetail);
+				}
+			}
+		}
+	}
+
+	if (VisualizationSettings.bShowParticles)
+	{
+		for (int32 Index = 0; Index != SimulationState.Particles.Num(); ++Index)
+		{
+			const RigParticleSimulation::FParticle& Particle = SimulationState.Particles[Index];
+			const RigParticleSimulation::FParticleCollider& ParticleCollider = SimulationState.ParticleColliders[Index];
+
+			FTransform ParticleComponentTM = Hierarchy.GetGlobalTransform(ParticleOwnerComponents[Index].GetElementKey());
+			ParticleComponentTM.SetTranslation(
+				SimulationSpaceState.ConvertSimSpacePositionToComponentSpace(FVector(Particle.Position)));
+			DI->DrawSphere(
+				ParticleComponentTM,
+				FTransform(),			
+				ParticleCollider.Radius,
+				Particle.InvMass != 0.0f
+				? RigDynamicsSolverDraw::SimulatedParticleColor
+				: RigDynamicsSolverDraw::KinematicParticleColor,
+				VisualizationSettings.LineThickness,
+				VisualizationSettings.ShapeDetail);
+		}
+	}
+
+}
+
diff --git a/Engine/Plugins/Experimental/ControlRigDynamics/Source/ControlRigDynamics/Private/RigParticleSimulation.cpp b/Engine/Plugins/Experimental/ControlRigDynamics/Source/ControlRigDynamics/Private/RigParticleSimulation.cpp
new file mode 100644
--- /dev/null
+++ b/Engine/Plugins/Experimental/ControlRigDynamics/Source/ControlRigDynamics/Private/RigParticleSimulation.cpp
@@ -0,0 +1,984 @@
+// Copyright Epic Games, Inc. All Rights Reserved.
+
+#include "RigParticleSimulation.h"
+
+//======================================================================================================================
+// This should be kept completely independent of control rig etc
+//======================================================================================================================
+
+namespace RigParticleSimulation
+{
+
+//======================================================================================================================
+void FParticleTarget::Update(
+	const FParticle&  InParticle,
+	const FParticle*  InParentParticle,
+	const float       InStrength,
+	const float       InDampingRatio,
+	const float       InExtraDamping,
+	const float       InTargetVelocityInfluence,
+	const float       InTargetMode)
+{
+	TargetVelocityInfluence = InTargetVelocityInfluence;
+	TargetMode = InTargetMode;
+
+	// Scale compliance by the effective inverse mass so the constraint frequency is
+	// mass-independent. Parent contribution is weighted by Alpha^2 (from the squared Jacobian).
+	const float Alpha = InParentParticle ? InTargetMode : 0.0f;
+	const float InvMass = InParticle.InvMass
+		+ Alpha * Alpha * (InParentParticle ? InParentParticle->InvMass : 0.0f);
+
+	// Angular frequency
+	const float W = InStrength * TWO_PI;
+	// Stiffness = W^2, and compliance is inverse stiffness (scaled by inverse mass)
+	Compliance = InvMass / FMath::Max(W * W, KINDA_SMALL_NUMBER);
+	Damping = (2.0f * InDampingRatio * W + InExtraDamping) / FMath::Max(InvMass, KINDA_SMALL_NUMBER);
+}
+
+//======================================================================================================================
+void FSoftDistanceConstraint::Update(
+	const FParticle& ParentParticle, 
+	const FParticle& ChildParticle,
+	const float      InTargetDistance,
+	const float      InStrength,
+	const float      InDampingRatio,
+	const float      InExtraDamping)
+{
+	TargetDistance = InTargetDistance;
+
+	// Scale compliance by effective inverse mass so the constraint frequency is
+	// mass-independent
+	const float W = InStrength * TWO_PI;
+	const float InvMassSum = ParentParticle.InvMass + ChildParticle.InvMass;
+	Compliance = InvMassSum / FMath::Max(W * W, KINDA_SMALL_NUMBER);
+	Damping = (2.0f * InDampingRatio * W + InExtraDamping) / FMath::Max(InvMassSum, KINDA_SMALL_NUMBER);
+}
+
+//======================================================================================================================
+static void StepParticleDynamic(
+	FParticle&                      Particle,
+	const FParticleInfo&            Info,
+	const float                     InDeltaTime,
+	const float                     InDtRatio,
+	const float                     InStepAlpha,
+	const FSimVector&               InGravity,
+	const FSimulationSpaceMotion&   InSpaceMotion)
+{
+	check(Particle.InvMass != 0.0f);
+
+	Particle.TargetPosition = FMath::Lerp(Info.PreviousTargetPosition, Info.TargetPosition, InStepAlpha);
+	Particle.TargetVelocity = FMath::Lerp(Info.PreviousTargetVelocity, Info.TargetVelocity, InStepAlpha);
+
+	// Time-corrected Verlet integration. Standard Verlet assumes constant dt:
+	//   NewPos = 2*Pos - PrevPos + Accel*dt^2
+	// which is equivalent to:
+	//   NewPos = Pos + (Pos - PrevPos) + Accel*dt^2
+	// The term (Pos - PrevPos) is the implicit velocity * dt_prev. When the timestep changes
+	// (e.g. at frame boundaries due to variable framerate or substepping changes), we scale the
+	// velocity term by (dt_current / dt_prev) to account for the different time interval.
+	FSimVector Acceleration = InGravity + Info.ExternalForce * Particle.InvMass;
+
+	// Non-inertial reference frame pseudo-forces. These are all mass-proportional (the
+	// acceleration is the same regardless of particle mass, like gravity).
+	const float SpaceMovementAmount = InSpaceMotion.SpaceMovementAmount;
+	if (SpaceMovementAmount > KINDA_SMALL_NUMBER)
+	{
+		const FSimVector R = Particle.Position;
+		const FSimVector& W = InSpaceMotion.AngularVelocity;
+
+		FSimVector PseudoAccel = FSimVector::ZeroVector;
+
+		if (InSpaceMotion.bAccelerationsValid)
+		{
+			// Fictitious force from linear acceleration: -a_frame
+			PseudoAccel -= InSpaceMotion.LinearAcceleration;
+			// Euler force from angular acceleration: -alpha x r
+			PseudoAccel -= FSimVector::CrossProduct(InSpaceMotion.AngularAcceleration, R);
+		}
+
+		// Centrifugal: -w x (w x r)
+		const FSimVector WxR = FSimVector::CrossProduct(W, R);
+		PseudoAccel -= FSimVector::CrossProduct(W, WxR);
+
+		// Coriolis: -2(w x v), using the implicit Verlet velocity
+		const FSimVector V = (Particle.Position - Particle.PrevPosition) / InDeltaTime;
+		PseudoAccel -= 2.0f * FSimVector::CrossProduct(W, V);
+
+		Acceleration += PseudoAccel * SpaceMovementAmount;
+
+		// External drag force (e.g. wind), mass-dependent: lighter particles are more affected
+		Acceleration += InSpaceMotion.ExternalDragForce * Particle.InvMass;
+	}
+
+	const FSimVector OldPosition = Particle.Position;
+	Particle.Position += 
+		(Particle.Position - Particle.PrevPosition) * InDtRatio + Acceleration * FMath::Square(InDeltaTime);
+	Particle.PrevPosition = OldPosition;
+
+	// Ether drag: blend velocity toward the "ether" (world-at-rest) velocity.
+	// Mass-dependent: lighter particles couple more strongly to the ether, like wind
+	// affecting a feather more than a rock.
+	if (SpaceMovementAmount > KINDA_SMALL_NUMBER)
+	{
+		const float LinDrag = InSpaceMotion.LinearDragMultiplier;
+		const float AngDrag = InSpaceMotion.AngularDragMultiplier;
+
+		if (LinDrag > KINDA_SMALL_NUMBER || AngDrag > KINDA_SMALL_NUMBER)
+		{
+			const FSimVector R = Particle.Position;
+			const FSimVector& W = InSpaceMotion.AngularVelocity;
+			const FSimVector WxR = FSimVector::CrossProduct(W, R);
+
+			// Ether velocity: what the particle's velocity would be if at rest in the world, scaled
+			// by the linear and angular components
+			const FSimVector ScaledEtherVel = -InSpaceMotion.LinearVelocity * LinDrag - WxR * AngDrag;
+			const FSimVector CurrentVel = (Particle.Position - Particle.PrevPosition) / InDeltaTime;
+
+			// Exponential blend: rate depends on InvMass (lighter = faster coupling)
+			// and timestep for frame-rate independence
+			const float DragRate = 1.0f - FMath::Exp(-SpaceMovementAmount * Particle.InvMass * InDeltaTime);
+
+			const FSimVector NewVel = FMath::Lerp(CurrentVel, ScaledEtherVel, DragRate);
+			Particle.PrevPosition = Particle.Position - NewVel * InDeltaTime;
+		}
+	}
+}
+
+//======================================================================================================================
+static void StepParticleKinematic(
+	FParticle&           Particle,
+	const FParticleInfo& Info,
+	const float          InDeltaTime,
+	const float          InStepAlpha)
+{
+	check(Particle.InvMass == 0.0f);
+
+	Particle.TargetPosition = FMath::Lerp(Info.PreviousTargetPosition, Info.TargetPosition, InStepAlpha);
+	Particle.TargetVelocity = FMath::Lerp(Info.PreviousTargetVelocity, Info.TargetVelocity, InStepAlpha);
+
+	// Kinematic - just track the target position
+	Particle.Position = Particle.TargetPosition;
+	Particle.PrevPosition = Particle.Position - Particle.TargetVelocity * InDeltaTime;
+}
+
+//======================================================================================================================
+static void StepParticleTarget(
+	FParticleTarget& ParticleTarget, const float InStepAlpha,
+	const FSimVector& InPrevTargetDirectionFromParent, const FSimVector& InTargetDirectionFromParent)
+{
+	// Lerp + normalize is a good approximation of Slerp for the small inter-substep angles
+	const FSimVector InterpDir =
+		FMath::Lerp(InPrevTargetDirectionFromParent, InTargetDirectionFromParent, InStepAlpha);
+	const float InterpLen = InterpDir.Size();
+	ParticleTarget.TargetDirectionFromParent =
+		(InterpLen > KINDA_SMALL_NUMBER) ? InterpDir * (1.0f / InterpLen) : FSimVector::ZeroVector;
+	ParticleTarget.Lambda = FSimVector::ZeroVector;
+}
+
+//======================================================================================================================
+static void StepSoftDistanceConstraint(FSoftDistanceConstraint& Constraint)
+{
+	Constraint.Lambda = 0.0f;
+}
+
+//======================================================================================================================
+static void StepCollider(FCollider& Collider, const float StepAlpha)
+{
+	Collider.TM.Blend(Collider.PreviousTM, Collider.TargetTM, StepAlpha);
+}
+
+//======================================================================================================================
+static void SolveTargetConstraint(
+	FParticle&        Particle,
+	FParticleTarget&  ParticleTarget,
+	FParticle*        ParentParticlePtr,
+	const bool        bOnlyMoveChild,
+	const float       DeltaTime,
+	const float       InvDeltaTime)
+{
+	// Note that we should never be called if InvMass = 0 (the caller should be sensible!)
+	check(Particle.InvMass > 0.0f);
+
+	// Alpha blends between SimSpace (0) and directional (1). If there's no parent, force
+	// SimSpace regardless of the setting.
+	const float Alpha = ParentParticlePtr ? ParticleTarget.TargetMode : 0.0f;
+
+	// SimSpace constraint error: drive child toward its absolute target position
+	const FSimVector C_sim = Particle.Position - Particle.TargetPosition;
+
+	// Directional constraint error: drive parent-child direction to match animation
+	FSimVector C_dir = FSimVector::ZeroVector;
+	if (Alpha > KINDA_SMALL_NUMBER && ParentParticlePtr)
+	{
+		const FSimVector Delta = Particle.Position - ParentParticlePtr->Position;
+		const float DeltaLen = Delta.Size();
+		if (DeltaLen > KINDA_SMALL_NUMBER
+			&& !ParticleTarget.TargetDirectionFromParent.IsNearlyZero(KINDA_SMALL_NUMBER))
+		{
+			C_dir = Delta - ParticleTarget.TargetDirectionFromParent * DeltaLen;
+		}
+	}
+
+	// Blended constraint error
+	const FSimVector C = FMath::Lerp(C_sim, C_dir, Alpha);
+
+	// Bilateral mass: w_child + Alpha^2 * w_parent
+	const float ParentInvMass = ParentParticlePtr ? ParentParticlePtr->InvMass : 0.0f;
+	const float InvMassForDenom =
+		bOnlyMoveChild ? Particle.InvMass : (Particle.InvMass + Alpha * Alpha * ParentInvMass);
+
+	// Blended velocity: child velocity minus Alpha-scaled parent velocity
+	const FSimVector ChildVel = (Particle.Position - Particle.PrevPosition) * InvDeltaTime;
+	FSimVector CurrentVelocity = ChildVel;
+	if (Alpha > KINDA_SMALL_NUMBER && ParentParticlePtr)
+	{
+		const FSimVector ParentVel =
+			(ParentParticlePtr->Position - ParentParticlePtr->PrevPosition) * InvDeltaTime;
+		CurrentVelocity -= Alpha * ParentVel;
+	}
+
+	FSimVector RefVelocity = Particle.TargetVelocity * ParticleTarget.TargetVelocityInfluence;
+	if (Alpha > KINDA_SMALL_NUMBER && ParentParticlePtr)
+	{
+		RefVelocity -= Alpha * ParentParticlePtr->TargetVelocity * ParticleTarget.TargetVelocityInfluence;
+	}
+
+	// A = alpha~ in the Mueller paper
+	const float A = ParticleTarget.Compliance * FMath::Square(InvDeltaTime);
+	// G = gamma in the Mueller paper = alpha~ * beta~ / dt
+	const float G = ParticleTarget.Compliance * ParticleTarget.Damping * InvDeltaTime;
+
+	const FSimVector RelVelocity = CurrentVelocity - RefVelocity;
+
+	// XPBD update (Jacobian = I for child, -Alpha*I for parent)
+	const float Denom = (1.0f + G) * InvMassForDenom + A;
+	const FSimVector DeltaLambda =
+		(-C - A * ParticleTarget.Lambda - G * RelVelocity * DeltaTime) / Denom;
+
+	ParticleTarget.Lambda += DeltaLambda;
+	Particle.Position += DeltaLambda * Particle.InvMass;
+
+	// Parent correction scaled by Alpha
+	if (Alpha > KINDA_SMALL_NUMBER && ParentParticlePtr && !bOnlyMoveChild)
+	{
+		ParentParticlePtr->Position -= Alpha * DeltaLambda * ParentParticlePtr->InvMass;
+	}
+}
+
+//======================================================================================================================
+static void SolveHardDistanceConstraint(
+	const FHardDistanceConstraint& Constraint,
+	FParticle&                     ParentParticle,
+	FParticle&                     ChildParticle,
+	const bool                     bOnlyMoveChild)
+{
+	if (bOnlyMoveChild)
+	{
+		if (ChildParticle.InvMass == 0)
+		{
+			return;
+		}
+
+		const FSimVector Delta = ChildParticle.Position - ParentParticle.Position;
+		const float DeltaLen = Delta.Size();
+		if (DeltaLen < KINDA_SMALL_NUMBER)
+		{
+			return;
+		}
+
+		const FSimVector N = Delta * (1.0f / DeltaLen);
+		ChildParticle.Position += N * (Constraint.TargetDistance - DeltaLen);
+	}
+	else
+	{
+		const FSimVector Delta = ChildParticle.Position - ParentParticle.Position;
+		FSimScalar DeltaLen = Delta.Size();
+		if (DeltaLen < KINDA_SMALL_NUMBER)
+		{
+			return;
+		}
+
+		const float InvMassParent = ParentParticle.InvMass;
+		const float InvMassChild = ChildParticle.InvMass;
+		const float InvMassSum = InvMassParent + InvMassChild;
+		if (InvMassSum == 0)
+		{
+			return;
+		}
+
+		const FSimVector N = Delta * (1.0f / DeltaLen);
+		const FSimVector Correction = N * ((Constraint.TargetDistance - DeltaLen) / InvMassSum);
+
+		ParentParticle.Position -= Correction * InvMassParent;
+		ChildParticle.Position += Correction * InvMassChild;
+	}
+}
+
+//======================================================================================================================
+static void SolveSoftDistanceConstraint(
+	FSoftDistanceConstraint& Constraint,
+	FParticle&               ParentParticle,
+	FParticle&               ChildParticle,
+	const float              DeltaTime,
+	const float              InvDeltaTime)
+{
+	const float InvMassParent = ParentParticle.InvMass;
+	const float InvMassChild = ChildParticle.InvMass;
+	const float InvMassSum = InvMassParent + InvMassChild;
+	if (InvMassSum == 0)
+	{
+		return;
+	}
+
+	const FSimVector Delta = ChildParticle.Position - ParentParticle.Position;
+	const float DeltaLen = Delta.Size();
+	if (DeltaLen < KINDA_SMALL_NUMBER)
+	{
+		return;
+	}
+
+	const FSimVector N = Delta * (1.0f / DeltaLen);
+	const float C = DeltaLen - Constraint.TargetDistance; // constraint error
+
+	// A = alpha~ in the Mueller paper
+	const float A = Constraint.Compliance * FMath::Square(InvDeltaTime);
+	// The paper uses B = beta~ = FMath::Square(DeltaTime) * Constraint.Damping;
+	// G = gamma in the Mueller paper = alpha~ * beta~ / dt
+	const float G = Constraint.Compliance * Constraint.Damping * InvDeltaTime;
+
+	// Compute relative velocity from positions (no stored velocity field)
+	const FSimVector RelVelocity =
+		(ChildParticle.Position - ChildParticle.PrevPosition
+			- ParentParticle.Position + ParentParticle.PrevPosition) * InvDeltaTime;
+	const float dCdt = N.Dot(RelVelocity);
+
+	// Solve
+	const float Num = -C - A * Constraint.Lambda - G * dCdt * DeltaTime;
+	const float Denom = (1.0f + G) * InvMassSum + A;
+	const float DeltaLambda = Num / Denom;
+
+	Constraint.Lambda += DeltaLambda;
+
+	// Apply
+	FSimVector Correction = N * DeltaLambda;
+	ParentParticle.Position -= Correction * InvMassParent;
+	ChildParticle.Position += Correction * InvMassChild;
+}
+
+//======================================================================================================================
+// Note that the docs in the Collider Components note that scale is ignored
+static void SolveContactConstraint(
+	const FBoxCollider&       BoxCollider,
+	const FSimTransform&      ColliderTM,
+	FParticle&                Particle,
+	const FParticleCollider& ParticleCollider)
+{
+	const FSimTransform TM = BoxCollider.TM * ColliderTM;
+
+	const FSimVector HalfExtents = BoxCollider.Extents * 0.5f;
+	const FSimVector LocalPos = TM.InverseTransformPositionNoScale(Particle.Position);
+
+	// Signed distance from particle center to each box face (positive = inside the box)
+	const float DPosX = HalfExtents.X - LocalPos.X;
+	const float DNegX = HalfExtents.X + LocalPos.X;
+	const float DPosY = HalfExtents.Y - LocalPos.Y;
+	const float DNegY = HalfExtents.Y + LocalPos.Y;
+	const float DPosZ = HalfExtents.Z - LocalPos.Z;
+	const float DNegZ = HalfExtents.Z + LocalPos.Z;
+
+	FSimVector N;
+	float C;
+
+	if (DPosX >= 0.0f && DNegX >= 0.0f &&
+		DPosY >= 0.0f && DNegY >= 0.0f &&
+		DPosZ >= 0.0f && DNegZ >= 0.0f)
+	{
+		// Particle center is inside (or on surface of) the box - push out through nearest face
+		float MinDist = DPosX;
+		FSimVector LocalNormal(1.0f, 0.0f, 0.0f);
+
+		if (DNegX < MinDist) {
+			MinDist = DNegX; 
+			LocalNormal = FSimVector(-1.0f, 0.0f, 0.0f); 
+		}
+		if (DPosY < MinDist)
+		{ 
+			MinDist = DPosY; 
+			LocalNormal = FSimVector(0.0f, 1.0f, 0.0f); 
+		}
+		if (DNegY < MinDist) 
+		{ 
+			MinDist = DNegY; 
+			LocalNormal = FSimVector(0.0f, -1.0f, 0.0f);
+		}
+		if (DPosZ < MinDist) 
+		{
+			MinDist = DPosZ; 
+			LocalNormal = FSimVector(0.0f, 0.0f, 1.0f); 
+		}
+		if (DNegZ < MinDist) 
+		{
+			MinDist = DNegZ; 
+			LocalNormal = FSimVector(0.0f, 0.0f, -1.0f); 
+		}
+
+		N = TM.TransformVectorNoScale(LocalNormal);
+		C = MinDist + ParticleCollider.Radius;
+	}
+	else
+	{
+		// Particle center is outside the box - find closest point on box surface
+		const FSimVector ClampedPos(
+			FMath::Clamp(LocalPos.X, -HalfExtents.X, HalfExtents.X),
+			FMath::Clamp(LocalPos.Y, -HalfExtents.Y, HalfExtents.Y),
+			FMath::Clamp(LocalPos.Z, -HalfExtents.Z, HalfExtents.Z));
+		const FSimVector ClosestPoint = TM.TransformPositionNoScale(ClampedPos);
+		const FSimVector Delta = Particle.Position - ClosestPoint;
+		const FSimScalar DeltaLen = Delta.Size();
+		if (DeltaLen > ParticleCollider.Radius || DeltaLen < KINDA_SMALL_NUMBER)
+		{
+			return;
+		}
+
+		N = Delta * (1.0f / DeltaLen);
+		C = ParticleCollider.Radius - DeltaLen;
+	}
+
+	Particle.Position += N * C;
+}
+
+//======================================================================================================================
+// Note that the docs in the Collider Components note that scale is ignored
+static void SolveContactConstraint(
+	const FCapsuleCollider&   CapsuleCollider,
+	const FSimTransform&      ColliderTM,
+	FParticle&                Particle,
+	const FParticleCollider&  ParticleCollider)
+{
+	const FSimTransform TM = CapsuleCollider.TM * ColliderTM;
+
+	const float MinDistance = CapsuleCollider.Radius + ParticleCollider.Radius;
+	if (MinDistance <= 0.0f)
+	{
+		return;
+	}
+
+	FSimVector SegmentPoint;
+	if (CapsuleCollider.Length < KINDA_SMALL_NUMBER)
+	{
+		SegmentPoint = TM.GetLocation();
+	}
+	else
+	{
+		// Inline closest-point-on-segment to avoid potential FVector3f overload issues
+		const FSimVector SegA = TM.TransformPositionNoScale(FSimVector(0, 0, -CapsuleCollider.Length * 0.5f));
+		const FSimVector SegB = TM.TransformPositionNoScale(FSimVector(0, 0, CapsuleCollider.Length * 0.5f));
+		const FSimVector Seg = SegB - SegA;
+		const float T = FMath::Clamp(
+			FSimVector::DotProduct(Particle.Position - SegA, Seg) /
+			FMath::Max(FSimVector::DotProduct(Seg, Seg), KINDA_SMALL_NUMBER),
+			0.0f, 1.0f);
+		SegmentPoint = SegA + T * Seg;
+	}
+
+	const FSimVector Delta = Particle.Position - SegmentPoint;
+	const float DeltaLen = Delta.Size();
+	if (DeltaLen > MinDistance || DeltaLen < KINDA_SMALL_NUMBER)
+	{
+		return;
+	}
+
+	const FSimVector N = Delta * (1.0f / DeltaLen);
+	const float C = MinDistance - DeltaLen;
+
+	Particle.Position += N * C;
+}
+
+//======================================================================================================================
+// Note that the docs in the Collider Components note that scale is ignored
+static void SolveContactConstraint(
+	const FPlaneCollider&     PlaneCollider,
+	const FSimTransform&      ColliderTM,
+	FParticle&                Particle,
+	const FParticleCollider&  ParticleCollider)
+{
+	const FSimTransform TM = PlaneCollider.TM * ColliderTM;
+
+	const FSimVector LocalPos = TM.InverseTransformPositionNoScale(Particle.Position);
+
+	// Plane normal is +Z in local space. Signed distance is positive on the normal side.
+	const float SignedDist = LocalPos.Z;
+	if (SignedDist >= ParticleCollider.Radius)
+	{
+		return;
+	}
+
+	// Reject if particle center is outside
+	const FSimVector2D HalfExtent = PlaneCollider.Extents * 0.5f;
+	if (FMath::Abs(LocalPos.X) > HalfExtent.X || FMath::Abs(LocalPos.Y) > HalfExtent.Y)
+	{
+		return;
+	}
+
+	// Push-out is in the plane normal direction
+	const FSimVector N = TM.TransformVectorNoScale(FSimVector(0, 0, 1));
+	const float C = ParticleCollider.Radius - SignedDist;
+
+	Particle.Position += N * C;
+}
+
+//======================================================================================================================
+static void SolveColliderContacts(
+	FParticle&                  Particle,
+	const FParticleCollider&    ParticleCollider,
+	const FParticleNoCollision& ParticleNoCollision,
+	TArrayView<const FCollider> Colliders)
+{
+	if (Particle.InvMass == 0.0f || !ParticleCollider.bCollideWithColliders)
+	{
+		return;
+	}
+	// Note that all the contact constraints here are hard and have no friction. This makes them
+	// very simple and fast, not requiring any per-contact state (i.e. no lambda for the normal and
+	// tangential force). We should anticipate adding a more fully-featured contact constraint that
+	// supports compliance and friction, enabled on a per-collider basis.
+	for (int32 ColliderIndex = 0; ColliderIndex < Colliders.Num(); ++ColliderIndex)
+	{
+		if (ParticleNoCollision.NoCollisionColliderIndices[ColliderIndex])
+		{
+			continue;
+		}
+
+		const FCollider& Collider = Colliders[ColliderIndex];
+
+		for (const FBoxCollider& BoxCollider : Collider.BoxColliders)
+		{
+			SolveContactConstraint(BoxCollider, Collider.TM, Particle, ParticleCollider);
+		}
+		for (const FCapsuleCollider& CapsuleCollider : Collider.CapsuleColliders)
+		{
+			SolveContactConstraint(CapsuleCollider, Collider.TM, Particle, ParticleCollider);
+		}
+		for (const FPlaneCollider& PlaneCollider : Collider.PlaneColliders)
+		{
+			SolveContactConstraint(PlaneCollider, Collider.TM, Particle, ParticleCollider);
+		}
+	}
+}
+
+//======================================================================================================================
+static void SolveParticleToParticleContact(
+	FParticle& ParticleA, const FParticleCollider& ParticleACollider, 
+	FParticle& ParticleB, const FParticleCollider& ParticleBCollider)
+{
+	const FSimScalar CombinedRadius = ParticleACollider.Radius + ParticleBCollider.Radius;
+
+	const FSimVector Delta = ParticleA.Position - ParticleB.Position;
+	FSimScalar DeltaLen = Delta.Size();
+	if (DeltaLen > CombinedRadius || DeltaLen < KINDA_SMALL_NUMBER)
+	{
+		return;
+	}
+
+	const float InvMassA = ParticleA.InvMass;
+	const float InvMassB = ParticleB.InvMass;
+	const float InvMassSum = InvMassA + InvMassB;
+	if (InvMassSum == 0)
+	{
+		return;
+	}
+
+	const FSimVector N = Delta * (1.0f / DeltaLen);
+	const FSimVector Correction = N * ((CombinedRadius - DeltaLen) / InvMassSum);
+
+	ParticleA.Position += Correction * InvMassA;
+	ParticleB.Position -= Correction * InvMassB;
+}
+
+//======================================================================================================================
+// The functions above here should be generic, just dealing with individual components of a simulation.
+// 
+// The functions below here are directly related to FSimulationState which is a specific simulation
+// configuration, oriented towards animation/skeletal dynamics.
+//======================================================================================================================
+
+//======================================================================================================================
+void UpdateParticle(
+	FSimulationState& State,
+	const int32       Index,
+	const float       InMass,
+	const float       InGravityMultiplier,
+	const float       InRadius)
+{
+	FParticleInfo& Info = State.ParticleInfos[Index];
+	FParticle& Particle = State.Particles[Index];
+	FParticleCollider& ParticleCollider = State.ParticleColliders[Index];
+
+	Info.GravityMultiplier = InGravityMultiplier;
+
+	const float InvMass =
+		Info.MovementType == ERigParticleSimulationMovementType::Simulated
+		? 1.0f / FMath::Max(InMass, KINDA_SMALL_NUMBER)
+		: 0.0f;
+
+	Particle.Update(InvMass);
+	ParticleCollider.Update(InRadius);
+}
+
+//======================================================================================================================
+void UpdateParticleTarget(
+	FSimulationState&  State,
+	const int32        Index,
+	float              InvDeltaTime,
+	const FSimVector&  InTargetPosition,
+	const float        InStrength,
+	const float        InDampingRatio,
+	const float        InExtraDamping,
+	const float        InTargetVelocityInfluence,
+	const float        InTargetMode)
+{
+	FParticleInfo& Info = State.ParticleInfos[Index];
+	const FParticle& Particle = State.Particles[Index];
+	FParticleTarget& ParticleTarget = State.ParticleTargets[Index];
+
+	// Shift previous targets
+	Info.PreviousTargetPosition = Info.TargetPosition;
+	Info.PreviousTargetVelocity = Info.TargetVelocity;
+
+	// Set new target and compute velocity
+	Info.TargetPosition = InTargetPosition;
+	Info.TargetVelocity = (Info.TargetPosition - Info.PreviousTargetPosition) * InvDeltaTime;
+
+	// Compute and store the normalized direction from parent to this particle in the target pose.
+	// Only needed when TargetMode > 0 (directional contribution).
+	// It is intentional that both directions use the current ParentParticleTargetPosition.
+	if (InTargetMode > KINDA_SMALL_NUMBER
+		&& State.ParticleInfos.IsValidIndex(Info.ParentParticleIndex))
+	{
+		const FSimVector& ParentTargetPos = State.ParticleInfos[Info.ParentParticleIndex].TargetPosition;
+
+		auto SafeNormalize = [](const FSimVector& V) -> FSimVector
+		{
+			const float Len = V.Size();
+			return (Len > KINDA_SMALL_NUMBER) ? V * (1.0f / Len) : FSimVector::ZeroVector;
+		};
+		Info.PreviousTargetDirectionFromParent =
+			SafeNormalize(Info.PreviousTargetPosition - ParentTargetPos);
+		Info.TargetDirectionFromParent = SafeNormalize(Info.TargetPosition - ParentTargetPos);
+	}
+
+	const FParticle* ParentParticle =
+		State.Particles.IsValidIndex(Info.ParentParticleIndex)
+		? &State.Particles[Info.ParentParticleIndex] : nullptr;
+
+	ParticleTarget.Update(
+		Particle, ParentParticle, InStrength, InDampingRatio, InExtraDamping,
+		InTargetVelocityInfluence, InTargetMode);
+}
+
+//======================================================================================================================
+void UpdateDynamicsStep(
+	FSimulationState&              SimulationState,
+	const FSolverSettings&         SolverSettings,
+	const FSimulationSpaceMotion&  SimulationSpaceForces,
+	const float                    DeltaTime,
+	const int32                    StepIndex,
+	const int32                    NumSteps)
+{
+	if (DeltaTime < SMALL_NUMBER)
+	{
+		return;
+	}
+	const float StepAlpha = (float)(StepIndex + 1) / (float)NumSteps;
+	const float InvDeltaTime = 1.0f / DeltaTime;
+
+	// Time-corrected Verlet: scale the implicit velocity by dt_current / dt_prev
+	const float DtRatio =
+		(SimulationState.PrevStepTime > KINDA_SMALL_NUMBER) ? (DeltaTime / SimulationState.PrevStepTime) : 1.0f;
+
+	// Step all the things - this means to update them for the (sub) step, prior to solving constraints
+	for (FSoftDistanceConstraint& Constraint : SimulationState.SoftDistanceConstraints)
+	{
+		StepSoftDistanceConstraint(Constraint);
+	}
+
+	// Hard distance constraints don't have a step
+
+	for (FCollider& Collider : SimulationState.Colliders)
+	{
+		StepCollider(Collider, StepAlpha);
+	}
+
+	for (int32 Index = 0; Index < SimulationState.Particles.Num(); ++Index)
+	{
+		const FParticleInfo& Info = SimulationState.ParticleInfos[Index];
+		FParticle& Particle = SimulationState.Particles[Index];
+		FParticleTarget& ParticleTarget = SimulationState.ParticleTargets[Index];
+
+		if (Particle.InvMass != 0.0f)
+		{
+			StepParticleDynamic(Particle, Info, DeltaTime, DtRatio, StepAlpha,
+				FSimVector(Info.GravityMultiplier * SolverSettings.Gravity), SimulationSpaceForces);
+			StepParticleTarget(ParticleTarget, StepAlpha, Info.PreviousTargetDirectionFromParent,
+				Info.TargetDirectionFromParent);
+		}
+		else
+		{
+			StepParticleKinematic(Particle, Info, DeltaTime, StepAlpha);
+		}
+	}
+
+	// Iteratively solve constraints
+
+	for (int32 Iter = 0; Iter != SolverSettings.NumIterations; ++Iter)
+	{
+		const bool bTargetOnlyMoveChild = (Iter == SolverSettings.NumIterations - 1);
+		for (int32 Index = 0; Index < SimulationState.Particles.Num(); ++Index)
+		{
+			FParticle& Particle = SimulationState.Particles[Index];
+			if (Particle.InvMass != 0.0f)
+			{
+				FParticleTarget& ParticleTarget = SimulationState.ParticleTargets[Index];
+				const FParticleInfo& Info = SimulationState.ParticleInfos[Index];
+				FParticle* ParentParticlePtr =
+					SimulationState.Particles.IsValidIndex(Info.ParentParticleIndex)
+					? &SimulationState.Particles[Info.ParentParticleIndex] : nullptr;
+				SolveTargetConstraint(
+					Particle, ParticleTarget, ParentParticlePtr, bTargetOnlyMoveChild,
+					DeltaTime, InvDeltaTime);
+			}
+		}
+
+		// Distance constraints etc
+		for (int32 ConstraintIter = 0; ConstraintIter != SolverSettings.NumConstraintSubIterations; ++ConstraintIter)
+		{
+			// Soft constraints can be soft
+			for (int32 Index = 0; Index != SimulationState.SoftDistanceConstraints.Num(); ++Index)
+			{
+				const FSoftDistanceConstraintInfo& Info =
+					SimulationState.SoftDistanceConstraintInfos[Index];
+				FSoftDistanceConstraint& Constraint =
+					SimulationState.SoftDistanceConstraints[Index];
+
+				check(SimulationState.Particles.IsValidIndex(Info.ParentIndex) &&
+					SimulationState.Particles.IsValidIndex(Info.ChildIndex));
+				SolveSoftDistanceConstraint(
+					Constraint,
+					SimulationState.Particles[Info.ParentIndex],
+					SimulationState.Particles[Info.ChildIndex],
+					DeltaTime,
+					InvDeltaTime);
+			}
+
+			// Hard distance constraints come later to reduce stretching
+			for (int32 Index = 0; Index != SimulationState.HardDistanceConstraints.Num(); ++Index)
+			{
+				const FHardDistanceConstraintInfo& Info = SimulationState.HardDistanceConstraintInfos[Index];
+				const FHardDistanceConstraint& Constraint = SimulationState.HardDistanceConstraints[Index];
+
+				check(SimulationState.Particles.IsValidIndex(Info.ParentIndex) &&
+					SimulationState.Particles.IsValidIndex(Info.ChildIndex));
+				SolveHardDistanceConstraint(
+					Constraint,
+					SimulationState.Particles[Info.ParentIndex],
+					SimulationState.Particles[Info.ChildIndex],
+					false);
+			}
+
+			// Collisions between particles
+			for (int32 ParticleIndex = 0 ; ParticleIndex != SimulationState.Particles.Num() ; ++ParticleIndex)
+			{
+				FParticle& Particle = SimulationState.Particles[ParticleIndex];
+				const FParticleCollider& ParticleCollider = SimulationState.ParticleColliders[ParticleIndex];
+				const FParticleToParticleCollision& ParticleToParticleCollision =
+					SimulationState.ParticleToParticleCollision[ParticleIndex];
+				for (int32 OtherParticleIndex : ParticleToParticleCollision.CollisionParticleIndices)
+				{
+					check(SimulationState.Particles.IsValidIndex(OtherParticleIndex));
+					FParticle& OtherParticle = SimulationState.Particles[OtherParticleIndex];
+					const FParticleCollider& OtherParticleCollider = SimulationState.ParticleColliders[OtherParticleIndex];
+					SolveParticleToParticleContact(
+						Particle, ParticleCollider, OtherParticle, OtherParticleCollider);
+				}
+			}
+
+			// Collisions should be as hard as possible
+			for (int32 ParticleIndex = 0 ; ParticleIndex != SimulationState.Particles.Num() ; ++ParticleIndex)
+			{
+				FParticle& Particle = SimulationState.Particles[ParticleIndex];
+				const FParticleCollider& ParticleCollider = SimulationState.ParticleColliders[ParticleIndex];
+				const FParticleNoCollision& ParticleNoCollision = SimulationState.ParticleNoCollision[ParticleIndex];
+				SolveColliderContacts(Particle, ParticleCollider, ParticleNoCollision, SimulationState.Colliders);
+			}
+
+			// Skeletal distance constraints come last to eliminate stretching. We also make them
+			// one-way on the last iteration.
+			bool bOnlyMoveChild =
+				(Iter == SolverSettings.NumIterations - 1) &&
+				(ConstraintIter == SolverSettings.NumConstraintSubIterations - 1);
+			for (int32 Index = 0; Index != SimulationState.SkeletalConstraints.Num(); ++Index)
+			{
+				const FHardDistanceConstraintInfo& Info = SimulationState.SkeletalConstraintInfos[Index];
+				const FHardDistanceConstraint& Constraint = SimulationState.SkeletalConstraints[Index];
+
+				check(SimulationState.Particles.IsValidIndex(Info.ParentIndex) &&
+					SimulationState.Particles.IsValidIndex(Info.ChildIndex));
+				SolveHardDistanceConstraint(
+					Constraint,
+					SimulationState.Particles[Info.ParentIndex],
+					SimulationState.Particles[Info.ChildIndex],
+					bOnlyMoveChild);
+			}
+		}
+	}
+}
+
+//======================================================================================================================
+void Simulate(
+	FSimulationState&              SimulationState,
+	const FSolverSettings&         SolverSettings,
+	const FSimulationSpaceMotion&  SimulationSpaceForces,
+	const float                    DeltaTime)
+{
+	QUICK_SCOPE_CYCLE_COUNTER(STAT_RigParticleSimulation_Simulate);
+
+	if (SolverSettings.MaxNumSteps < 1)
+	{
+		return;
+	}
+
+	int32 NumSteps = FMath::TruncToInt(DeltaTime / FMath::Max(SolverSettings.MaxTimeStep, KINDA_SMALL_NUMBER));
+	NumSteps = FMath::Clamp(NumSteps + 1, 1, SolverSettings.MaxNumSteps);
+
+	float StepTime = DeltaTime / NumSteps; // Safe as we clamp above
+	for (int32 StepIndex = 0; StepIndex != NumSteps; ++StepIndex)
+	{
+		UpdateDynamicsStep(SimulationState, SolverSettings, SimulationSpaceForces, StepTime, StepIndex, NumSteps);
+		SimulationState.PrevStepTime = StepTime;
+	}
+}
+
+//======================================================================================================================
+TArray<int32> SortParticlesRootToLeaf(FSimulationState& State)
+{
+	const int32 NumParticles = State.Particles.Num();
+	if (NumParticles <= 1)
+	{
+		TArray<int32> Identity;
+		if (NumParticles == 1)
+		{
+			Identity.Add(0);
+		}
+		return Identity;
+	}
+
+	// Build children-of adjacency lists from ParentParticleIndex
+	TArray<TArray<int32>> Children;
+	Children.SetNum(NumParticles);
+
+	TArray<int32> Queue;
+	Queue.Reserve(NumParticles);
+
+	for (int32 ParticleIndex = 0; ParticleIndex < NumParticles; ++ParticleIndex)
+	{
+		const int32 ParentIndex = State.ParticleInfos[ParticleIndex].ParentParticleIndex;
+		if (ParentIndex == INDEX_NONE)
+		{
+			Queue.Add(ParticleIndex); // Root - seed the BFS queue
+		}
+		else
+		{
+			check(State.ParticleInfos.IsValidIndex(ParentIndex));
+			Children[ParentIndex].Add(ParticleIndex);
+		}
+	}
+
+	// Breadth First Search (BFS) from all roots - visitation order gives root-to-leaf sorted order
+	TArray<int32> SortedOrder;
+	SortedOrder.Reserve(NumParticles);
+
+	int32 Head = 0;
+	while (Head < Queue.Num())
+	{
+		const int32 Current = Queue[Head++];
+		SortedOrder.Add(Current);
+		for (const int32 Child : Children[Current])
+		{
+			Queue.Add(Child);
+		}
+	}
+	// Cycles are impossible: ParentParticleIndex is derived from the skeleton hierarchy (a tree),
+	// so the particle parent graph is always a forest and all particles are reachable from roots.
+	check(SortedOrder.Num() == NumParticles);
+
+	// Compute inverse mapping: OldToNew[old_index] = new_index
+	TArray<int32> OldToNew;
+	OldToNew.SetNum(NumParticles);
+	for (int32 NewIndex = 0; NewIndex < NumParticles; ++NewIndex)
+	{
+		OldToNew[SortedOrder[NewIndex]] = NewIndex;
+	}
+
+	// Reorder all particle-parallel arrays using the gather pattern
+	{
+		TArray<FParticleInfo> TempInfos;                         TempInfos.SetNum(NumParticles);
+		TArray<FParticle> TempParticles;                         TempParticles.SetNum(NumParticles);
+		TArray<FParticleTarget> TempTargets;                     TempTargets.SetNum(NumParticles);
+		TArray<FParticleCollider> TempColliders;                 TempColliders.SetNum(NumParticles);
+		TArray<FParticleNoCollision> TempNoCollision;            TempNoCollision.SetNum(NumParticles);
+		TArray<FParticleToParticleCollision> TempP2PCollision;   TempP2PCollision.SetNum(NumParticles);
+
+		for (int32 NewIndex = 0; NewIndex < NumParticles; ++NewIndex)
+		{
+			const int32 OldIndex = SortedOrder[NewIndex];
+			TempInfos[NewIndex]        = MoveTemp(State.ParticleInfos[OldIndex]);
+			TempParticles[NewIndex]    = MoveTemp(State.Particles[OldIndex]);
+			TempTargets[NewIndex]      = MoveTemp(State.ParticleTargets[OldIndex]);
+			TempColliders[NewIndex]    = MoveTemp(State.ParticleColliders[OldIndex]);
+			TempNoCollision[NewIndex]  = MoveTemp(State.ParticleNoCollision[OldIndex]);
+			TempP2PCollision[NewIndex] = MoveTemp(State.ParticleToParticleCollision[OldIndex]);
+		}
+
+		State.ParticleInfos               = MoveTemp(TempInfos);
+		State.Particles                   = MoveTemp(TempParticles);
+		State.ParticleTargets             = MoveTemp(TempTargets);
+		State.ParticleColliders           = MoveTemp(TempColliders);
+		State.ParticleNoCollision         = MoveTemp(TempNoCollision);
+		State.ParticleToParticleCollision = MoveTemp(TempP2PCollision);
+	}
+
+	// Remap all stored particle indices using OldToNew. Note that at the current call site
+	// (Instantiate, before constraints are created), the constraint and collision-pair arrays
+	// are empty. We remap them anyway so the function is correct regardless of when it's called.
+	for (int32 I = 0; I < NumParticles; ++I)
+	{
+		int32& ParentIndex = State.ParticleInfos[I].ParentParticleIndex;
+		if (ParentIndex != INDEX_NONE)
+		{
+			ParentIndex = OldToNew[ParentIndex];
+		}
+		for (int32& OtherIndex : State.ParticleToParticleCollision[I].CollisionParticleIndices)
+		{
+			OtherIndex = OldToNew[OtherIndex];
+		}
+	}
+	for (FHardDistanceConstraintInfo& Info : State.SkeletalConstraintInfos)
+	{
+		Info.ParentIndex = OldToNew[Info.ParentIndex];
+		Info.ChildIndex = OldToNew[Info.ChildIndex];
+	}
+	for (FHardDistanceConstraintInfo& Info : State.HardDistanceConstraintInfos)
+	{
+		Info.ParentIndex = OldToNew[Info.ParentIndex];
+		Info.ChildIndex = OldToNew[Info.ChildIndex];
+	}
+	for (FSoftDistanceConstraintInfo& Info : State.SoftDistanceConstraintInfos)
+	{
+		Info.ParentIndex = OldToNew[Info.ParentIndex];
+		Info.ChildIndex = OldToNew[Info.ChildIndex];
+	}
+
+	return SortedOrder;
+}
+
+}
diff --git a/Engine/Plugins/Experimental/ControlRigDynamics/Source/ControlRigDynamics/Public/ControlRigDynamicsModule.h b/Engine/Plugins/Experimental/ControlRigDynamics/Source/ControlRigDynamics/Public/ControlRigDynamicsModule.h
new file mode 100644
--- /dev/null
+++ b/Engine/Plugins/Experimental/ControlRigDynamics/Source/ControlRigDynamics/Public/ControlRigDynamicsModule.h
@@ -0,0 +1,23 @@
+// Copyright Epic Games, Inc. All Rights Reserved.
+
+#pragma once
+
+#include "CoreMinimal.h"
+#include "Modules/ModuleManager.h"
+
+#if UE_BUILD_SHIPPING || UE_BUILD_TEST
+DECLARE_LOG_CATEGORY_EXTERN(LogControlRigDynamics, Warning, Warning);
+#else
+DECLARE_LOG_CATEGORY_EXTERN(LogControlRigDynamics, Log, All);
+#endif
+
+//======================================================================================================================
+class FControlRigDynamicsModule : public IModuleInterface
+{
+public:
+
+	/** IModuleInterface implementation */
+	virtual void StartupModule() override;
+	virtual void ShutdownModule() override;
+};
+
diff --git a/Engine/Plugins/Experimental/ControlRigDynamics/Source/ControlRigDynamics/Public/RigDynamicsColliderComponent.h b/Engine/Plugins/Experimental/ControlRigDynamics/Source/ControlRigDynamics/Public/RigDynamicsColliderComponent.h
new file mode 100644
--- /dev/null
+++ b/Engine/Plugins/Experimental/ControlRigDynamics/Source/ControlRigDynamics/Public/RigDynamicsColliderComponent.h
@@ -0,0 +1,34 @@
+// Copyright Epic Games, Inc. All Rights Reserved.
+
+#pragma once
+
+#include "Rigs/RigHierarchyComponents.h"
+
+#include "RigDynamicsData.h"
+
+#include "RigDynamicsColliderComponent.generated.h"
+
+#define UE_API CONTROLRIGDYNAMICS_API
+
+//======================================================================================================================
+// A component that can be added to a joint/element that defines how a collision shape can be
+// attached to it. DynamicsParticles will collide against these shapes.
+//======================================================================================================================
+USTRUCT(BlueprintType)
+struct FRigDynamicsColliderComponent : public FRigBaseComponent
+{
+public:
+	GENERATED_BODY()
+	DECLARE_RIG_COMPONENT_METHODS(FRigDynamicsColliderComponent)
+
+	UPROPERTY(BlueprintReadOnly, EditAnywhere, Category = Shapes, meta = (ShowOnlyInnerProperties))
+	FRigDynamicsColliderCollection Colliders;
+
+	UE_API virtual void Save(FArchive& Ar) override;
+	UE_API virtual void Load(FArchive& Ar) override;
+
+	virtual FName GetDefaultComponentName() const override { return GetDefaultName(); }
+	static FName GetDefaultName() { return TEXT("DynamicsCollider"); }
+};
+
+#undef UE_API
diff --git a/Engine/Plugins/Experimental/ControlRigDynamics/Source/ControlRigDynamics/Public/RigDynamicsColliderExecution.h b/Engine/Plugins/Experimental/ControlRigDynamics/Source/ControlRigDynamics/Public/RigDynamicsColliderExecution.h
new file mode 100644
--- /dev/null
+++ b/Engine/Plugins/Experimental/ControlRigDynamics/Source/ControlRigDynamics/Public/RigDynamicsColliderExecution.h
@@ -0,0 +1,272 @@
+// Copyright Epic Games, Inc. All Rights Reserved.
+
+#pragma once
+
+#include "RigDynamicsExecution.h"
+#include "RigDynamicsColliderComponent.h"
+#include "RigDynamicsSolverComponent.h"
+
+#include "RigDynamicsColliderExecution.generated.h"
+
+#define UE_API CONTROLRIGDYNAMICS_API
+
+//======================================================================================================================
+// Spawns a new dynamics collider as a component on the owner element. These colliders track the
+// input pose, but do not track the simulated pose - so don't expect good results if you enable
+// collisions between particles and colliders that are attached to bones that are themselves simulated.
+// Note: This node only runs as part of the construction event.
+USTRUCT(meta = (DisplayName = "Spawn Dynamics Collider", Keywords = "Add,Construction,Create,New,Shape,Collision", Varying))
+struct FRigUnit_SpawnDynamicsCollider : public FRigUnit_DynamicsBaseMutable
+{
+	GENERATED_BODY()
+
+	RIGVM_METHOD()
+	UE_API virtual void Execute() override;
+
+	// The owner of the newly created component (must be set/valid)
+	UPROPERTY(meta = (Input, BoneName))
+	FRigElementKey Owner = FRigElementKey(TEXT(""), ERigElementType::Bone);
+
+	// Optional solver - if set (and valid), it will be added to this solver component
+	UPROPERTY(meta = (Input))
+	FRigComponentKey DynamicsSolverComponentKey = FRigComponentKey(
+		FRigElementKey(TEXT("Root"), ERigElementType::Bone), FRigDynamicsSolverComponent::GetDefaultName());
+
+	// The collider shapes
+	UPROPERTY(meta = (Input))
+	FRigDynamicsColliderCollection Colliders;
+
+	// The Dynamics Collider component key that was created
+	UPROPERTY(meta = (Output))
+	FRigComponentKey DynamicsColliderComponentKey;
+};
+
+//======================================================================================================================
+// Sets the properties of a named box shape on a Dynamics collider. Name is used to find the shape.
+USTRUCT(meta = (DisplayName = "Set Dynamics Collider Box", Varying))
+struct FRigUnit_HierarchySetDynamicsColliderBox : public FRigUnit_DynamicsBaseMutable
+{
+	GENERATED_BODY()
+
+	RIGVM_METHOD()
+	UE_API virtual void Execute() override;
+
+	// The Dynamics Collider to be updated
+	UPROPERTY(meta = (Input))
+	FRigComponentKey DynamicsColliderComponentKey = FRigComponentKey(
+		FRigElementKey(TEXT(""), ERigElementType::Bone), FRigDynamicsColliderComponent::GetDefaultName());
+
+	// The box shape data (Name is used to find the shape to update)
+	UPROPERTY(meta = (Input))
+	FRigDynamicsColliderBox Box;
+};
+
+//======================================================================================================================
+// Sets the properties of a named capsule shape on a Dynamics collider. Name is used to find the shape.
+USTRUCT(meta = (DisplayName = "Set Dynamics Collider Capsule", Varying))
+struct FRigUnit_HierarchySetDynamicsColliderCapsule : public FRigUnit_DynamicsBaseMutable
+{
+	GENERATED_BODY()
+
+	RIGVM_METHOD()
+	UE_API virtual void Execute() override;
+
+	// The Dynamics Collider to be updated
+	UPROPERTY(meta = (Input))
+	FRigComponentKey DynamicsColliderComponentKey = FRigComponentKey(
+		FRigElementKey(TEXT(""), ERigElementType::Bone), FRigDynamicsColliderComponent::GetDefaultName());
+
+	// The capsule shape data (Name is used to find the shape to update)
+	UPROPERTY(meta = (Input))
+	FRigDynamicsColliderCapsule Capsule;
+};
+
+//======================================================================================================================
+// Sets the properties of a named plane shape on a Dynamics collider. Name is used to find the shape.
+USTRUCT(meta = (DisplayName = "Set Dynamics Collider Plane", Varying))
+struct FRigUnit_HierarchySetDynamicsColliderPlane : public FRigUnit_DynamicsBaseMutable
+{
+	GENERATED_BODY()
+
+	RIGVM_METHOD()
+	UE_API virtual void Execute() override;
+
+	// The Dynamics Collider to be updated
+	UPROPERTY(meta = (Input))
+	FRigComponentKey DynamicsColliderComponentKey = FRigComponentKey(
+		FRigElementKey(TEXT(""), ERigElementType::Bone), FRigDynamicsColliderComponent::GetDefaultName());
+
+	// The plane shape data (Name is used to find the shape to update)
+	UPROPERTY(meta = (Input))
+	FRigDynamicsColliderPlane Plane;
+};
+
+//======================================================================================================================
+// Sets the transform of a named box shape on a Dynamics collider
+USTRUCT(meta = (DisplayName = "Set Dynamics Collider Box TM", Varying))
+struct FRigUnit_HierarchySetDynamicsColliderBoxTM : public FRigUnit_DynamicsBaseMutable
+{
+	GENERATED_BODY()
+
+	RIGVM_METHOD()
+	UE_API virtual void Execute() override;
+
+	// The Dynamics Collider to be updated
+	UPROPERTY(meta = (Input))
+	FRigComponentKey DynamicsColliderComponentKey = FRigComponentKey(
+		FRigElementKey(TEXT(""), ERigElementType::Bone), FRigDynamicsColliderComponent::GetDefaultName());
+
+	// The name of the box shape to update
+	UPROPERTY(meta = (Input))
+	FName Name;
+
+	// Position and orientation (scale is ignored)
+	UPROPERTY(meta = (Input))
+	FTransform TM;
+};
+
+//======================================================================================================================
+// Sets the extents of a named box shape on a Dynamics collider
+USTRUCT(meta = (DisplayName = "Set Dynamics Collider Box Extents", Varying))
+struct FRigUnit_HierarchySetDynamicsColliderBoxExtents : public FRigUnit_DynamicsBaseMutable
+{
+	GENERATED_BODY()
+
+	RIGVM_METHOD()
+	UE_API virtual void Execute() override;
+
+	// The Dynamics Collider to be updated
+	UPROPERTY(meta = (Input))
+	FRigComponentKey DynamicsColliderComponentKey = FRigComponentKey(
+		FRigElementKey(TEXT(""), ERigElementType::Bone), FRigDynamicsColliderComponent::GetDefaultName());
+
+	// The name of the box shape to update
+	UPROPERTY(meta = (Input))
+	FName Name;
+
+	// Full extents of the box in each axis
+	UPROPERTY(meta = (Input))
+	FVector Extents = FVector::OneVector * RigDynamicsDefaultColliderSize;
+};
+
+//======================================================================================================================
+// Sets the transform of a named capsule shape on a Dynamics collider
+USTRUCT(meta = (DisplayName = "Set Dynamics Collider Capsule TM", Varying))
+struct FRigUnit_HierarchySetDynamicsColliderCapsuleTM : public FRigUnit_DynamicsBaseMutable
+{
+	GENERATED_BODY()
+
+	RIGVM_METHOD()
+	UE_API virtual void Execute() override;
+
+	// The Dynamics Collider to be updated
+	UPROPERTY(meta = (Input))
+	FRigComponentKey DynamicsColliderComponentKey = FRigComponentKey(
+		FRigElementKey(TEXT(""), ERigElementType::Bone), FRigDynamicsColliderComponent::GetDefaultName());
+
+	// The name of the capsule shape to update
+	UPROPERTY(meta = (Input))
+	FName Name;
+
+	// Position and orientation (scale is ignored), main axis along +Z
+	UPROPERTY(meta = (Input))
+	FTransform TM;
+};
+
+//======================================================================================================================
+// Sets the radius of a named capsule shape on a Dynamics collider
+USTRUCT(meta = (DisplayName = "Set Dynamics Collider Capsule Radius", Varying))
+struct FRigUnit_HierarchySetDynamicsColliderCapsuleRadius : public FRigUnit_DynamicsBaseMutable
+{
+	GENERATED_BODY()
+
+	RIGVM_METHOD()
+	UE_API virtual void Execute() override;
+
+	// The Dynamics Collider to be updated
+	UPROPERTY(meta = (Input))
+	FRigComponentKey DynamicsColliderComponentKey = FRigComponentKey(
+		FRigElementKey(TEXT(""), ERigElementType::Bone), FRigDynamicsColliderComponent::GetDefaultName());
+
+	// The name of the capsule shape to update
+	UPROPERTY(meta = (Input))
+	FName Name;
+
+	// Capsule radius
+	UPROPERTY(meta = (Input, ClampMin = "0.0"))
+	float Radius = RigDynamicsDefaultColliderSize;
+};
+
+//======================================================================================================================
+// Sets the length of a named capsule shape on a Dynamics collider
+USTRUCT(meta = (DisplayName = "Set Dynamics Collider Capsule Length", Varying))
+struct FRigUnit_HierarchySetDynamicsColliderCapsuleLength : public FRigUnit_DynamicsBaseMutable
+{
+	GENERATED_BODY()
+
+	RIGVM_METHOD()
+	UE_API virtual void Execute() override;
+
+	// The Dynamics Collider to be updated
+	UPROPERTY(meta = (Input))
+	FRigComponentKey DynamicsColliderComponentKey = FRigComponentKey(
+		FRigElementKey(TEXT(""), ERigElementType::Bone), FRigDynamicsColliderComponent::GetDefaultName());
+
+	// The name of the capsule shape to update
+	UPROPERTY(meta = (Input))
+	FName Name;
+
+	// Core length of the capsule (total length is Length + 2 * Radius)
+	UPROPERTY(meta = (Input, ClampMin = "0.0"))
+	float Length = RigDynamicsDefaultColliderSize;
+};
+
+//======================================================================================================================
+// Sets the transform of a named plane shape on a Dynamics collider
+USTRUCT(meta = (DisplayName = "Set Dynamics Collider Plane TM", Varying))
+struct FRigUnit_HierarchySetDynamicsColliderPlaneTM : public FRigUnit_DynamicsBaseMutable
+{
+	GENERATED_BODY()
+
+	RIGVM_METHOD()
+	UE_API virtual void Execute() override;
+
+	// The Dynamics Collider to be updated
+	UPROPERTY(meta = (Input))
+	FRigComponentKey DynamicsColliderComponentKey = FRigComponentKey(
+		FRigElementKey(TEXT(""), ERigElementType::Bone), FRigDynamicsColliderComponent::GetDefaultName());
+
+	// The name of the plane shape to update
+	UPROPERTY(meta = (Input))
+	FName Name;
+
+	// Position defines a point on the plane, plane faces +Z
+	UPROPERTY(meta = (Input))
+	FTransform TM;
+};
+
+//======================================================================================================================
+// Sets the extents of a named plane shape on a Dynamics collider
+USTRUCT(meta = (DisplayName = "Set Dynamics Collider Plane Extents", Varying))
+struct FRigUnit_HierarchySetDynamicsColliderPlaneExtents : public FRigUnit_DynamicsBaseMutable
+{
+	GENERATED_BODY()
+
+	RIGVM_METHOD()
+	UE_API virtual void Execute() override;
+
+	// The Dynamics Collider to be updated
+	UPROPERTY(meta = (Input))
+	FRigComponentKey DynamicsColliderComponentKey = FRigComponentKey(
+		FRigElementKey(TEXT(""), ERigElementType::Bone), FRigDynamicsColliderComponent::GetDefaultName());
+
+	// The name of the plane shape to update
+	UPROPERTY(meta = (Input))
+	FName Name;
+
+	// Extents of the plane
+	UPROPERTY(meta = (Input, ClampMin = "0.0"))
+	FVector2D Extents = FVector2D(RigDynamicsDefaultColliderSize * 10.0f);
+};
+
+#undef UE_API
diff --git a/Engine/Plugins/Experimental/ControlRigDynamics/Source/ControlRigDynamics/Public/RigDynamicsConstraintComponent.h b/Engine/Plugins/Experimental/ControlRigDynamics/Source/ControlRigDynamics/Public/RigDynamicsConstraintComponent.h
new file mode 100644
--- /dev/null
+++ b/Engine/Plugins/Experimental/ControlRigDynamics/Source/ControlRigDynamics/Public/RigDynamicsConstraintComponent.h
@@ -0,0 +1,68 @@
+// Copyright Epic Games, Inc. All Rights Reserved.
+
+#pragma once
+
+#include "RigParticleSimulation.h"
+
+#include "Rigs/RigHierarchyComponents.h"
+
+#include "RigDynamicsConstraintComponent.generated.h"
+
+#define UE_API CONTROLRIGDYNAMICS_API
+
+//======================================================================================================================
+// A component that can be added to a joint/element that defines a constraint between two particles. 
+//======================================================================================================================
+USTRUCT(BlueprintType)
+struct FRigDynamicsConstraintComponent : public FRigBaseComponent
+{
+public:
+	GENERATED_BODY()
+	DECLARE_RIG_COMPONENT_METHODS(FRigDynamicsConstraintComponent)
+
+	// The parent particle for the constraint
+	UPROPERTY(BlueprintReadOnly, EditAnywhere, Category = Constraint)
+	FRigComponentKey ParentComponentKey;
+
+	// The child particle for the constraint
+	UPROPERTY(BlueprintReadOnly, EditAnywhere, Category = Constraint)
+	FRigComponentKey ChildComponentKey;
+
+	// If the constraint is marked as "hard", then the system will try as hard as possible to
+	// enforce it. This will also be a little cheaper than making a soft constraint.
+	UPROPERTY(BlueprintReadOnly, EditAnywhere, Category = Constraint)
+	ERigParticleSimulationConstraintType ConstraintType = ERigParticleSimulationConstraintType::Hard;
+
+	// The strength which we drive towards the target length. This is the oscillation frequency, so
+	// low values will be soft and springy, but values significantly above 1/timestep will track
+	// the target very accurately.
+	UPROPERTY(BlueprintReadOnly, EditAnywhere, Category = Dynamics, meta = (ClampMin = "0.0", EditCondition = "ConstraintType != ERigParticleSimulationConstraintType::Hard"))
+	float Strength = 2.0f;
+
+	// DampingRatio for the constraint
+	UPROPERTY(BlueprintReadOnly, EditAnywhere, Category = Dynamics, meta = (ClampMin = "0.0", EditCondition = "ConstraintType != ERigParticleSimulationConstraintType::Hard"))
+	float DampingRatio = 1.0f;
+
+	// Additional damping for the constraint
+	UPROPERTY(BlueprintReadOnly, EditAnywhere, Category = Dynamics, meta = (ClampMin = "0.0", EditCondition = "ConstraintType != ERigParticleSimulationConstraintType::Hard"))
+	float ExtraDamping = 0.0f;
+
+	// The target length will be calculated automatically, multiplied by this
+	UPROPERTY(BlueprintReadOnly, EditAnywhere, Category = Dynamics, meta = (ClampMin = "0.0"))
+	float LengthMultiplier = 1.0f;
+
+	// This will be added on to the target length after the automatic/multiplication
+	UPROPERTY(BlueprintReadOnly, EditAnywhere, Category = Dynamics)
+	float ExtraLength = 0.0f;
+
+	UE_API virtual void Save(FArchive& Ar) override;
+	UE_API virtual void Load(FArchive& Ar) override;
+
+	virtual FName GetDefaultComponentName() const override { return GetDefaultName(); }
+	static FName GetDefaultName() { return TEXT("DynamicsConstraint"); }
+
+protected:
+	void Serialize(FArchive& Ar);
+};
+
+#undef UE_API
diff --git a/Engine/Plugins/Experimental/ControlRigDynamics/Source/ControlRigDynamics/Public/RigDynamicsConstraintExecution.h b/Engine/Plugins/Experimental/ControlRigDynamics/Source/ControlRigDynamics/Public/RigDynamicsConstraintExecution.h
new file mode 100644
--- /dev/null
+++ b/Engine/Plugins/Experimental/ControlRigDynamics/Source/ControlRigDynamics/Public/RigDynamicsConstraintExecution.h
@@ -0,0 +1,279 @@
+// Copyright Epic Games, Inc. All Rights Reserved.
+
+#pragma once
+
+#include "RigDynamicsExecution.h"
+#include "RigDynamicsConstraintComponent.h"
+#include "RigDynamicsSolverComponent.h"
+#include "RigDynamicsParticleComponent.h"
+
+#include "RigParticleSimulation.h"
+
+#include "RigDynamicsConstraintExecution.generated.h"
+
+#define UE_API CONTROLRIGDYNAMICS_API
+
+//======================================================================================================================
+// Spawns a new dynamics constraint component.
+// Note: This node only runs as part of the construction event.
+//======================================================================================================================
+USTRUCT(meta = (DisplayName = "Spawn Dynamics Constraint", Keywords = "Add,Construction,Create,New,Rod,Distance", Varying))
+struct FRigUnit_SpawnDynamicsConstraint : public FRigUnit_DynamicsBaseMutable
+{
+	GENERATED_BODY()
+
+	RIGVM_METHOD()
+	UE_API virtual void Execute() override;
+
+	// The owner of the newly created component (must be set/valid)
+	UPROPERTY(meta = (Input, BoneName))
+	FRigElementKey Owner = FRigElementKey(TEXT(""), ERigElementType::Bone);
+
+	// Optional solver - if set, it will be added to this solver component
+	UPROPERTY(meta = (Input))
+	FRigComponentKey DynamicsSolverComponentKey = FRigComponentKey(
+		FRigElementKey(TEXT(""), ERigElementType::Bone), FRigDynamicsSolverComponent::GetDefaultName());
+
+	// The parent particle for the constraint
+	UPROPERTY(meta = (Input))
+	FRigComponentKey ParentComponentKey = FRigComponentKey(
+		FRigElementKey(TEXT(""), ERigElementType::Bone), FRigDynamicsParticleComponent::GetDefaultName());
+
+	// The child particle for the constraint
+	UPROPERTY(meta = (Input))
+	FRigComponentKey ChildComponentKey = FRigComponentKey(
+		FRigElementKey(TEXT(""), ERigElementType::Bone), FRigDynamicsParticleComponent::GetDefaultName());
+
+	// If the constraint is marked as "hard", then the system will try as hard as possible to
+	// enforce it. This will also be a little cheaper than making a soft constraint.
+	UPROPERTY(meta = (Input))
+	ERigParticleSimulationConstraintType ConstraintType = ERigParticleSimulationConstraintType::Hard;
+
+	// The strength which we drive towards the target length. This is the oscillation frequency, so
+	// low values will be soft and springy, but values significantly above 1/timestep will track
+	// the target very accurately.
+	UPROPERTY(meta = (Input, ClampMin = "0.0", EditCondition = "ConstraintType != ERigParticleSimulationConstraintType::Hard"))
+	float Strength = 2.0;
+
+	// DampingRatio for the constraint
+	UPROPERTY(meta = (Input, ClampMin = "0.0", EditCondition = "ConstraintType != ERigParticleSimulationConstraintType::Hard"))
+	float DampingRatio = 1.0;
+
+	// Additional damping for the constraint
+	UPROPERTY(meta = (Input, ClampMin = "0.0", EditCondition = "ConstraintType != ERigParticleSimulationConstraintType::Hard"))
+	float ExtraDamping = 0.0;
+
+	// The target length will be calculated automatically, multiplied by this
+	UPROPERTY(meta = (Input, ClampMin = "0.0"))
+	float LengthMultiplier = 1.0f;
+
+	// This will be added on to the target length after the automatic/multiplication
+	UPROPERTY(meta = (Input))
+	float ExtraLength = 0.0f;
+
+	// The Dynamics Constraint component key that was created
+	UPROPERTY(meta = (Output))
+	FRigComponentKey DynamicsConstraintComponentKey;
+};
+
+//======================================================================================================================
+// Sets the strength of a Dynamics constraint
+USTRUCT(meta = (DisplayName = "Set Dynamics Constraint Strength", Varying))
+struct FRigUnit_HierarchySetDynamicsConstraintStrength : public FRigUnit_DynamicsBaseMutable
+{
+	GENERATED_BODY()
+
+	RIGVM_METHOD()
+	UE_API virtual void Execute() override;
+
+	// The Dynamics Constraint to be updated
+	UPROPERTY(meta = (Input))
+	FRigComponentKey DynamicsConstraintComponentKey = FRigComponentKey(
+		FRigElementKey(TEXT(""), ERigElementType::Bone), FRigDynamicsConstraintComponent::GetDefaultName());
+
+	// The strength which drives towards the target length (oscillation frequency)
+	UPROPERTY(meta = (Input, ClampMin = "0.0"))
+	float Strength = 2.0;
+};
+
+//======================================================================================================================
+// Gets the strength of a Dynamics constraint
+USTRUCT(meta = (DisplayName = "Get Dynamics Constraint Strength", Varying))
+struct FRigUnit_HierarchyGetDynamicsConstraintStrength : public FRigUnit_DynamicsBase
+{
+	GENERATED_BODY()
+
+	RIGVM_METHOD()
+	UE_API virtual void Execute() override;
+
+	// The Dynamics Constraint to query
+	UPROPERTY(meta = (Input))
+	FRigComponentKey DynamicsConstraintComponentKey = FRigComponentKey(
+		FRigElementKey(TEXT(""), ERigElementType::Bone), FRigDynamicsConstraintComponent::GetDefaultName());
+
+	// The strength which drives towards the target length (oscillation frequency)
+	UPROPERTY(meta = (Output))
+	float Strength = 0.0f;
+};
+
+//======================================================================================================================
+// Sets the damping ratio of a Dynamics constraint
+USTRUCT(meta = (DisplayName = "Set Dynamics Constraint Damping Ratio", Varying))
+struct FRigUnit_HierarchySetDynamicsConstraintDampingRatio : public FRigUnit_DynamicsBaseMutable
+{
+	GENERATED_BODY()
+
+	RIGVM_METHOD()
+	UE_API virtual void Execute() override;
+
+	// The Dynamics Constraint to be updated
+	UPROPERTY(meta = (Input))
+	FRigComponentKey DynamicsConstraintComponentKey = FRigComponentKey(
+		FRigElementKey(TEXT(""), ERigElementType::Bone), FRigDynamicsConstraintComponent::GetDefaultName());
+
+	// Damping ratio for the constraint
+	UPROPERTY(meta = (Input, ClampMin = "0.0"))
+	float DampingRatio = 1.0;
+};
+
+//======================================================================================================================
+// Gets the damping ratio of a Dynamics constraint
+USTRUCT(meta = (DisplayName = "Get Dynamics Constraint Damping Ratio", Varying))
+struct FRigUnit_HierarchyGetDynamicsConstraintDampingRatio : public FRigUnit_DynamicsBase
+{
+	GENERATED_BODY()
+
+	RIGVM_METHOD()
+	UE_API virtual void Execute() override;
+
+	// The Dynamics Constraint to query
+	UPROPERTY(meta = (Input))
+	FRigComponentKey DynamicsConstraintComponentKey = FRigComponentKey(
+		FRigElementKey(TEXT(""), ERigElementType::Bone), FRigDynamicsConstraintComponent::GetDefaultName());
+
+	// Damping ratio for the constraint
+	UPROPERTY(meta = (Output))
+	float DampingRatio = 0.0f;
+};
+
+//======================================================================================================================
+// Sets the extra damping of a Dynamics constraint
+USTRUCT(meta = (DisplayName = "Set Dynamics Constraint Extra Damping", Varying))
+struct FRigUnit_HierarchySetDynamicsConstraintExtraDamping : public FRigUnit_DynamicsBaseMutable
+{
+	GENERATED_BODY()
+
+	RIGVM_METHOD()
+	UE_API virtual void Execute() override;
+
+	// The Dynamics Constraint to be updated
+	UPROPERTY(meta = (Input))
+	FRigComponentKey DynamicsConstraintComponentKey = FRigComponentKey(
+		FRigElementKey(TEXT(""), ERigElementType::Bone), FRigDynamicsConstraintComponent::GetDefaultName());
+
+	// Additional damping for the constraint
+	UPROPERTY(meta = (Input, ClampMin = "0.0"))
+	float ExtraDamping = 0.0;
+};
+
+//======================================================================================================================
+// Gets the extra damping of a Dynamics constraint
+USTRUCT(meta = (DisplayName = "Get Dynamics Constraint Extra Damping", Varying))
+struct FRigUnit_HierarchyGetDynamicsConstraintExtraDamping : public FRigUnit_DynamicsBase
+{
+	GENERATED_BODY()
+
+	RIGVM_METHOD()
+	UE_API virtual void Execute() override;
+
+	// The Dynamics Constraint to query
+	UPROPERTY(meta = (Input))
+	FRigComponentKey DynamicsConstraintComponentKey = FRigComponentKey(
+		FRigElementKey(TEXT(""), ERigElementType::Bone), FRigDynamicsConstraintComponent::GetDefaultName());
+
+	// Additional damping for the constraint
+	UPROPERTY(meta = (Output))
+	float ExtraDamping = 0.0f;
+};
+
+//======================================================================================================================
+// Sets the length multiplier of a Dynamics constraint
+USTRUCT(meta = (DisplayName = "Set Dynamics Constraint Length Multiplier", Varying))
+struct FRigUnit_HierarchySetDynamicsConstraintLengthMultiplier : public FRigUnit_DynamicsBaseMutable
+{
+	GENERATED_BODY()
+
+	RIGVM_METHOD()
+	UE_API virtual void Execute() override;
+
+	// The Dynamics Constraint to be updated
+	UPROPERTY(meta = (Input))
+	FRigComponentKey DynamicsConstraintComponentKey = FRigComponentKey(
+		FRigElementKey(TEXT(""), ERigElementType::Bone), FRigDynamicsConstraintComponent::GetDefaultName());
+
+	// The target length will be calculated automatically, multiplied by this
+	UPROPERTY(meta = (Input, ClampMin = "0.0"))
+	float LengthMultiplier = 1.0f;
+};
+
+//======================================================================================================================
+// Gets the length multiplier of a Dynamics constraint
+USTRUCT(meta = (DisplayName = "Get Dynamics Constraint Length Multiplier", Varying))
+struct FRigUnit_HierarchyGetDynamicsConstraintLengthMultiplier : public FRigUnit_DynamicsBase
+{
+	GENERATED_BODY()
+
+	RIGVM_METHOD()
+	UE_API virtual void Execute() override;
+
+	// The Dynamics Constraint to query
+	UPROPERTY(meta = (Input))
+	FRigComponentKey DynamicsConstraintComponentKey = FRigComponentKey(
+		FRigElementKey(TEXT(""), ERigElementType::Bone), FRigDynamicsConstraintComponent::GetDefaultName());
+
+	// The target length will be calculated automatically, multiplied by this
+	UPROPERTY(meta = (Output))
+	float LengthMultiplier = 0.0f;
+};
+
+//======================================================================================================================
+// Sets the extra length of a Dynamics constraint
+USTRUCT(meta = (DisplayName = "Set Dynamics Constraint Extra Length", Varying))
+struct FRigUnit_HierarchySetDynamicsConstraintExtraLength : public FRigUnit_DynamicsBaseMutable
+{
+	GENERATED_BODY()
+
+	RIGVM_METHOD()
+	UE_API virtual void Execute() override;
+
+	// The Dynamics Constraint to be updated
+	UPROPERTY(meta = (Input))
+	FRigComponentKey DynamicsConstraintComponentKey = FRigComponentKey(
+		FRigElementKey(TEXT(""), ERigElementType::Bone), FRigDynamicsConstraintComponent::GetDefaultName());
+
+	// Added on to the target length after the automatic/multiplication
+	UPROPERTY(meta = (Input))
+	float ExtraLength = 0.0f;
+};
+
+//======================================================================================================================
+// Gets the extra length of a Dynamics constraint
+USTRUCT(meta = (DisplayName = "Get Dynamics Constraint Extra Length", Varying))
+struct FRigUnit_HierarchyGetDynamicsConstraintExtraLength : public FRigUnit_DynamicsBase
+{
+	GENERATED_BODY()
+
+	RIGVM_METHOD()
+	UE_API virtual void Execute() override;
+
+	// The Dynamics Constraint to query
+	UPROPERTY(meta = (Input))
+	FRigComponentKey DynamicsConstraintComponentKey = FRigComponentKey(
+		FRigElementKey(TEXT(""), ERigElementType::Bone), FRigDynamicsConstraintComponent::GetDefaultName());
+
+	// Added on to the target length after the automatic/multiplication
+	UPROPERTY(meta = (Output))
+	float ExtraLength = 0.0f;
+};
+
+#undef UE_API
diff --git a/Engine/Plugins/Experimental/ControlRigDynamics/Source/ControlRigDynamics/Public/RigDynamicsData.h b/Engine/Plugins/Experimental/ControlRigDynamics/Source/ControlRigDynamics/Public/RigDynamicsData.h
new file mode 100644
--- /dev/null
+++ b/Engine/Plugins/Experimental/ControlRigDynamics/Source/ControlRigDynamics/Public/RigDynamicsData.h
@@ -0,0 +1,354 @@
+// Copyright Epic Games, Inc. All Rights Reserved.
+
+#pragma once
+
+#include "Rigs/RigHierarchyComponents.h"
+
+#include "RigDynamicsData.generated.h"
+
+#define UE_API CONTROLRIGDYNAMICS_API
+
+constexpr float RigDynamicsDefaultColliderSize = 10.0f;
+
+//======================================================================================================================
+USTRUCT(BlueprintType)
+struct FRigDynamicsColliderBox
+{
+	GENERATED_BODY()
+
+	FRigDynamicsColliderBox() {}
+
+	FRigDynamicsColliderBox(const FName InName, const FTransform& InTM, const FVector& InExtents)
+		: Name(InName), TM(InTM), Extents(InExtents) {}
+
+	// Shape name is optional/only used for identification
+	UPROPERTY(BlueprintReadOnly, EditAnywhere, Category = Shape)
+	FName Name;
+
+	// Position and orientation (scale is ignored)
+	UPROPERTY(BlueprintReadOnly, EditAnywhere, Category = Shape, meta=(MakeEditWidget=true))
+	FTransform TM;
+
+	// These are the full extents of the box in each axis
+	UPROPERTY(BlueprintReadOnly, EditAnywhere, Category = Shape, meta = (ClampMin = "0.0"))
+	FVector Extents = FVector::OneVector * RigDynamicsDefaultColliderSize;
+};
+
+//======================================================================================================================
+USTRUCT(BlueprintType)
+struct FRigDynamicsColliderCapsule
+{
+	GENERATED_BODY()
+
+	FRigDynamicsColliderCapsule() {}
+
+	FRigDynamicsColliderCapsule(const FName InName, const FTransform& InTM, const float InRadius, const float InLength)
+		: Name(InName), TM(InTM), Radius(InRadius), Length(InLength) {}
+
+	// Shape name is optional/only used for identification
+	UPROPERTY(BlueprintReadOnly, EditAnywhere, Category = Shape)
+	FName Name;
+
+	// Position and orientation (scale is ignored)
+	// Main axis will be along +Z. 
+	UPROPERTY(BlueprintReadOnly, EditAnywhere, Category = Shape, meta=(MakeEditWidget=true))
+	FTransform TM;
+
+	// Capsule radius
+	UPROPERTY(BlueprintReadOnly, EditAnywhere, Category = Shape, meta = (ClampMin = "0.0"))
+	float Radius = RigDynamicsDefaultColliderSize;
+
+	// This is the length of the core part of the capsule. The total length will be Length + 2 * Radius
+	UPROPERTY(BlueprintReadOnly, EditAnywhere, Category = Shape, meta = (ClampMin = "0.0"))
+	float Length = RigDynamicsDefaultColliderSize;
+};
+
+//======================================================================================================================
+USTRUCT(BlueprintType)
+struct FRigDynamicsColliderPlane
+{
+	GENERATED_BODY()
+
+	FRigDynamicsColliderPlane() {}
+
+	FRigDynamicsColliderPlane(const FName InName, const FTransform& InTM, const FVector2D& InExtents)
+		: Name(InName), TM(InTM), Extents(InExtents) {}
+
+	// Shape name is optional/only used for identification
+	UPROPERTY(BlueprintReadOnly, EditAnywhere, Category = Shape)
+	FName Name;
+
+	// Position defines a point on the plane. Plane faces in the +ve Z direction
+	UPROPERTY(BlueprintReadOnly, EditAnywhere, Category = Shape, meta=(MakeEditWidget=true))
+	FTransform TM;
+
+	// Extents of the plane
+	UPROPERTY(BlueprintReadOnly, EditAnywhere, Category = Shape, meta = (ClampMin = "0.0"))
+	FVector2D Extents = FVector2D(RigDynamicsDefaultColliderSize * 10.0f);
+};
+
+//======================================================================================================================
+USTRUCT(BlueprintType)
+struct FRigDynamicsColliderCollection
+{
+	GENERATED_BODY()
+
+	UPROPERTY(BlueprintReadOnly, EditAnywhere, Category = Shapes)
+	TArray<FRigDynamicsColliderBox> Boxes;
+
+	UPROPERTY(BlueprintReadOnly, EditAnywhere, Category = Shapes)
+	TArray<FRigDynamicsColliderCapsule> Capsules;
+
+	UPROPERTY(BlueprintReadOnly, EditAnywhere, Category = Shapes)
+	TArray<FRigDynamicsColliderPlane> Planes;
+
+	bool IsEmpty() const { return Boxes.IsEmpty() && Capsules.IsEmpty() && Planes.IsEmpty(); }
+	int32 NumShapes() const { return Boxes.Num() + Capsules.Num() + Planes.Num(); }
+};
+
+//======================================================================================================================
+// Specifies what space the simulation should run in.
+UENUM()
+enum class ERigDynamicsSimulationSpace : uint8
+{
+	/** Simulate in world space. Moving the skeletal mesh will generate velocity changes */
+	World,
+	/** 
+	* Simulate in component space. Moving the entire skeletal mesh will have no effect on velocities 
+	* unless this is explicitly enabled in the simulation space settings */
+	Component,
+	/**
+	* Simulate in the space of the specified bone. Moving the entire skeletal mesh and/or the Bone
+	* will have no effect on velocities unless this is explicitly enabled in the simulation space settings.
+	*/
+	SpaceBone,
+};
+
+//======================================================================================================================
+// Configuration/settings for the dynamics solver
+USTRUCT(BlueprintType)
+struct FRigDynamicsSolverSettings
+{
+	GENERATED_BODY()
+
+	UPROPERTY(BlueprintReadOnly, EditAnywhere, Category = SolverSettings)
+	ERigDynamicsSimulationSpace SimulationSpace = ERigDynamicsSimulationSpace::Component;
+
+	// The space in which collision shapes are defined
+	// TODO use this
+	UPROPERTY(BlueprintReadOnly, EditAnywhere, Category = SolverSettings)
+	ERigDynamicsSimulationSpace CollisionSpace = ERigDynamicsSimulationSpace::Component;
+
+	// The bone to use for the collision or simulation space (only relevant if one is set to SpaceBone)
+	UPROPERTY(BlueprintReadOnly, EditAnywhere, Category = SolverSettings)
+	FRigElementKey SpaceBone = FRigElementKey(TEXT(""), ERigElementType::Bone);
+
+	// The collider shapes defined in the collision space - e.g. for representing a ground etc
+	// TODO use this
+	UPROPERTY(BlueprintReadOnly, EditAnywhere, Category = SolverSettings, meta = (ShowOnlyInnerProperties))
+	FRigDynamicsColliderCollection Colliders;
+
+	// Gravity in world space. This can be scaled by each particle.
+	UPROPERTY(BlueprintReadOnly, EditAnywhere, Category = SolverSettings)
+	FVector Gravity = FVector(0.0f, 0.0f, -981.0f);
+
+	// The maximum timestep of any step.
+	UPROPERTY(BlueprintReadOnly, EditAnywhere, Category = SolverSettings, meta = (ClampMin = "0.0"))
+	float MaxTimeStep = 1.0f;
+
+	// The maximum number of steps that the update can be divided into. If MaxNumSteps * MaxTimeStep
+	// is less than the update time, the simulation will run "slow".
+	UPROPERTY(BlueprintReadOnly, EditAnywhere, Category = SolverSettings, meta = (ClampMin = "0"))
+	int32 MaxNumSteps = 1;
+
+	// This is the total number of iterations, including iterating over the particle targets. This
+	// can normally be 1, unless you want to track targets with very accurate spring behaviour.
+	UPROPERTY(BlueprintReadOnly, EditAnywhere, Category = SolverSettings, meta = (ClampMin = "0"))
+	int32 NumIterations = 1;
+
+	// The number of constraint (distance and collision) sub-iterations. Along-chain constraints
+	// will be one-way on the last iteration, so it this is 1 (which is often enough), particles in
+	// a chain will never be able to affect their parents. Increasing will result in more realistic
+	// behavior that comes with a performance cost.
+	UPROPERTY(BlueprintReadOnly, EditAnywhere, Category = SolverSettings, meta = (ClampMin = "0"))
+	int32 NumConstraintSubIterations = 1;
+
+	// If true, then the read-back from the solver is done by rotating bones to match the directions
+	// between particles. If there is no branching, then this will result in the final bones
+	// matching the simulated particles (depending on the skeletal bone constraints being
+	// fulfilled). However, when there is branching, there can be some deviation.
+	UPROPERTY(BlueprintReadOnly, EditAnywhere, Category = SolverSettings)
+	bool bSetBoneOrientations = true;
+
+	// If true the bone positions (translation) will be updated based on the particles. This will
+	// introduce stretching etc in the skeleton that is not there in the original animation, but it
+	// should result in better matching of the simulation and the bone positions.
+	UPROPERTY(BlueprintReadOnly, EditAnywhere, Category = SolverSettings)
+	bool bSetBonePositions = false;
+};
+
+//======================================================================================================================
+// These settings are used when the simulation is not done in world space. They allow the movement
+// of the simulation space (e.g. the component) to be added to the local simulation in a controlled way.
+USTRUCT(BlueprintType)
+struct FRigDynamicsSimulationSpaceSettings
+{
+	GENERATED_BODY()
+
+	// Overall multiplier on the effects of simulation space movement on the simulation
+	UPROPERTY(BlueprintReadOnly, EditAnywhere, Category = Settings, meta = (ClampMin = "0.0", ClampMax = "1.0"))
+	float SpaceMovementAmount = 1.0f;
+
+	// Multiplier on the components of vertical velocity and acceleration of the simulation space
+	// that is passed to the simulation. Usually from 0.0 to 1.0 to reduce the effects of jumping
+	// and crouching on the simulation, but it can be higher than 1.0 if you need to exaggerate this
+	// motion for some reason.
+	UPROPERTY(BlueprintReadOnly, EditAnywhere, Category = Settings)
+	float VelocityScaleZ = 1.0f;
+
+	UPROPERTY(BlueprintReadOnly, EditAnywhere, Category = Settings, meta = (InlineEditConditionToggle))
+	bool bClampLinearVelocity = false;
+
+	// A clamp on the effective world-space velocity that is passed to the simulation. The default
+	// value effectively means "unlimited". It is not usually required to change this but you would
+	// reduce this to limit the effects of drag on the bodies in the simulation (if you have bodies
+	// that have LinearDrag set to non-zero in the physics asset). Expected values in this case
+	// would be somewhat less than the usual velocities of your object which is commonly a few
+	// hundred for a character.
+	UPROPERTY(BlueprintReadOnly, EditAnywhere, Category = Settings, meta = (ClampMin = "0.0", EditCondition = "bClampLinearVelocity"))
+	float MaxLinearVelocity = 10000.0f;
+
+	UPROPERTY(BlueprintReadOnly, EditAnywhere, Category = Settings, meta = (InlineEditConditionToggle))
+	bool bClampAngularVelocity = false;
+
+	// A clamp on the effective world-space angular velocity that is passed to the simulation. Units
+	// are radian/s, so a value of about 6.0 is one rotation per second. The default value
+	// effectively means "unlimited". You would reduce this (and MaxAngularAcceleration) to limit
+	// how much bodies "fly out" when the actor spins on the spot. This is especially useful if you
+	// have characters than can rotate very quickly and you would probably want values around or
+	// less than 10 in this case.
+	UPROPERTY(BlueprintReadOnly, EditAnywhere, Category = Settings, meta = (ClampMin = "0.0", EditCondition = "bClampAngularVelocity"))
+	float MaxAngularVelocity = 10000.0f;
+
+	UPROPERTY(BlueprintReadOnly, EditAnywhere, Category = Settings, meta = (InlineEditConditionToggle))
+	bool bClampLinearAcceleration = false;
+
+	// A clamp on the effective world-space acceleration that is passed to the simulation. The
+	// default value effectively means "unlimited". This property is used to stop the bodies of the
+	// simulation flying out when suddenly changing linear speed. It is useful when you have
+	// characters than can changes from stationary to running very quickly such as in an FPS. A
+	// common value for a character might be in the few hundreds.
+	UPROPERTY(BlueprintReadOnly, EditAnywhere, Category = Settings, meta = (ClampMin = "0.0", EditCondition = "bClampLinearAcceleration"))
+	float MaxLinearAcceleration = 10000.0f;
+
+	UPROPERTY(BlueprintReadOnly, EditAnywhere, Category = Settings, meta = (InlineEditConditionToggle))
+	bool bClampAngularAcceleration = false;
+
+	// A clamp on the effective world-space angular acceleration that is passed to the simulation.
+	// Units are radian/s/s. The default value effectively means "unlimited". This has a similar
+	// effect to MaxAngularVelocity, except that it is related to the flying out of bodies when the
+	// rotation speed suddenly changes. A typical value for a character might be around 100.
+	UPROPERTY(BlueprintReadOnly, EditAnywhere, Category = Settings, meta = (ClampMin = "0.0", EditCondition = "bClampAngularAcceleration"))
+	float MaxAngularAcceleration = 10000.0f;
+
+	// If the linear acceleration is above this threshold, the movement is treated as a teleport.
+	// The calculated accelerations and velocities will be set to zero, without resetting the
+	// simulation state. A value of zero will disable this detection. Note: Not implemented yet
+	UPROPERTY(BlueprintReadOnly, EditAnywhere, Category = Settings, meta = (ClampMin = "0.0"))
+	float LinearAccelerationThresholdForTeleport = 100000.0f;
+
+	// If the angular acceleration (deg/s/s) is above this threshold, the movement is treated as a
+	// teleport. The calculated accelerations and velocities will be set to zero, without resetting
+	// the simulation state. A value of zero will disable this detection. Note: Not implemented yet
+	UPROPERTY(BlueprintReadOnly, EditAnywhere, Category = Settings, meta = (ClampMin = "0.0"))
+	float AngularAccelerationThresholdForTeleport = 100000.0f;
+
+	// If the linear position change is above this threshold, the movement is treated as a teleport.
+	// The calculated accelerations and velocities will be set to zero, without resetting the
+	// simulation state. A value of zero will disable this detection. Note: Not implemented yet
+	UPROPERTY(BlueprintReadOnly, EditAnywhere, Category = Settings, meta = (ClampMin = "0.0"))
+	float PositionChangeThresholdForTeleport = 100.0f;
+
+	// If the angular change (degrees) is above this threshold, the movement is treated as a
+	// teleport. The calculated accelerations and velocities will be set to zero, without resetting
+	// the simulation state. A value of zero will disable this detection. Note: Not implemented yet
+	UPROPERTY(BlueprintReadOnly, EditAnywhere, Category = Settings, meta = (ClampMin = "0.0"))
+	float OrientationChangeThresholdForTeleport = 30.0f;
+
+	// How much of the simulation frame's linear velocity to pass onto the bodies (linear ether
+	// drag). This relies on there being drag authored on the bodies. If set to zero, the only drag
+	// will be from the body's local movement (in the simulation space).
+	UPROPERTY(BlueprintReadOnly, EditAnywhere, Category = Settings, meta = (ClampMin = "0.0", ClampMax = "1.0"))
+	float LinearDragMultiplier = 1.0f;
+
+	// How much of the simulation frame's angular velocity to pass onto the bodies (angular ether
+	// drag). This relies on there being drag authored on the bodies. If set to zero, the only drag
+	// will be from the body's local movement (in the simulation space).
+	UPROPERTY(BlueprintReadOnly, EditAnywhere, Category = Settings, meta = (ClampMin = "0.0", ClampMax = "1.0"))
+	float AngularDragMultiplier = 1.0f;
+
+	// Additional linear drag from movement of the simulation space, applied to every body in
+	// addition to linear drag specified on them in the physics asset. When combined with
+	// ExternalLinearVelocity, this can be used to add a temporary wind-blown effect without having
+	// to tune linear drag on all the bodies in the physics asset. The result is that each body has
+	// a force equal to -ExternalLinearDrag * ExternalLinearVelocity applied to it, in addition to
+	// all other forces. The vector is in the component's space.
+	UPROPERTY(BlueprintReadOnly, EditAnywhere, Category = Settings)
+	FVector ExternalLinearDrag = FVector::ZeroVector;
+
+	// Additional velocity that is added to the component velocity so the simulation acts as if the
+	// actor is moving at speed, even when stationary. The vector is in world space. This could be
+	// used for wind effects etc. Typical values are similar to the velocity of the object or
+	// effect, and usually around or less than 1000 for characters/wind.
+	UPROPERTY(BlueprintReadOnly, EditAnywhere, Category = Settings)
+	FVector ExternalLinearVelocity = FVector::ZeroVector;
+
+	// Additional angular velocity that is added to the component angular velocity. This can be used
+	// to make the simulation act as if the actor is rotating even when it is not. E.g., to apply
+	// physics to a character on a podium as the camera rotates around it, to emulate the podium
+	// itself rotating. Vector is in world space. Units are deg/s.
+	UPROPERTY(BlueprintReadOnly, EditAnywhere, Category = Settings)
+	FVector ExternalAngularVelocity = FVector::ZeroVector;
+
+	// This will treat the external velocity like a wind field and add turbulence to it. Units are
+	// the same as velocity, so this is the approximate magnitude of the turbulence. This will be
+	// added after any clamps etc are applied to the external velocity.
+	UPROPERTY(BlueprintReadOnly, EditAnywhere, Category = Settings)
+	FVector ExternalTurbulenceVelocity = FVector::ZeroVector;
+};
+
+//======================================================================================================================
+USTRUCT(BlueprintType)
+struct FRigDynamicsVisualizationSettings
+{
+	GENERATED_BODY()
+
+	// Enable visualization. There is a significant runtime cost to enabling visualization, so make
+	// sure you disable it when it's not needed, and when you are profiling. Note that this can be
+	// overridden using CVar ControlRig.Dynamics.VisualizationOverride
+	UPROPERTY(EditAnywhere, Category = Visualization)
+	bool bEnableVisualization = false;
+
+	UPROPERTY(EditAnywhere, Category = Visualization, meta = (ClampMin = "0.0"))
+	float LineThickness = 1.0f;
+
+	UPROPERTY(EditAnywhere, Category = Visualization, meta = (ClampMin = "0.0"))
+	int32 ShapeDetail = 16;
+
+	UPROPERTY(EditAnywhere, Category = Visualization)
+	bool bShowParticles = true;
+
+	UPROPERTY(EditAnywhere, Category = Visualization)
+	bool bShowColliders = true;
+};
+
+//======================================================================================================================
+UE_API FArchive& operator <<(FArchive& Ar, FRigDynamicsVisualizationSettings& Data);
+UE_API FArchive& operator <<(FArchive& Ar, FRigDynamicsSimulationSpaceSettings& Data);
+UE_API FArchive& operator <<(FArchive& Ar, FRigDynamicsSolverSettings& Data);
+UE_API FArchive& operator <<(FArchive& Ar, FRigDynamicsColliderCollection& Data);
+UE_API FArchive& operator <<(FArchive& Ar, FRigDynamicsColliderPlane& Data);
+UE_API FArchive& operator <<(FArchive& Ar, FRigDynamicsColliderCapsule& Data);
+UE_API FArchive& operator <<(FArchive& Ar, FRigDynamicsColliderBox& Data);
+
+#undef UE_API
diff --git a/Engine/Plugins/Experimental/ControlRigDynamics/Source/ControlRigDynamics/Public/RigDynamicsExecution.h b/Engine/Plugins/Experimental/ControlRigDynamics/Source/ControlRigDynamics/Public/RigDynamicsExecution.h
new file mode 100644
--- /dev/null
+++ b/Engine/Plugins/Experimental/ControlRigDynamics/Source/ControlRigDynamics/Public/RigDynamicsExecution.h
@@ -0,0 +1,30 @@
+// Copyright Epic Games, Inc. All Rights Reserved.
+
+#pragma once
+
+#include "Units/Execution/RigUnit_DynamicHierarchy.h"
+
+#include "RigDynamicsExecution.generated.h"
+
+#define UE_API CONTROLRIGDYNAMICS_API
+
+//======================================================================================================================
+// Base struct for all other mutable dynamics nodes
+//======================================================================================================================
+USTRUCT(meta = (Category = "RigDynamics", NodeColor = "0.6 1.0 0.3", Keywords = "Dynamics", DocumentationPolicy = "Strict"))
+struct FRigUnit_DynamicsBaseMutable : public FRigUnitMutable
+{
+	GENERATED_BODY()
+};
+
+//======================================================================================================================
+// Base struct for all other non-mutable dynamics nodes
+//======================================================================================================================
+USTRUCT(meta = (Category = "RigDynamics", NodeColor = "0.6 1.0 0.3", Keywords = "Dynamics", DocumentationPolicy = "Strict"))
+struct FRigUnit_DynamicsBase : public FRigUnit
+{
+	GENERATED_BODY()
+};
+
+#undef UE_API
+
diff --git a/Engine/Plugins/Experimental/ControlRigDynamics/Source/ControlRigDynamics/Public/RigDynamicsExecutionHelpers.h b/Engine/Plugins/Experimental/ControlRigDynamics/Source/ControlRigDynamics/Public/RigDynamicsExecutionHelpers.h
new file mode 100644
--- /dev/null
+++ b/Engine/Plugins/Experimental/ControlRigDynamics/Source/ControlRigDynamics/Public/RigDynamicsExecutionHelpers.h
@@ -0,0 +1,129 @@
+// Copyright Epic Games, Inc. All Rights Reserved.
+
+#pragma once
+
+#include "RigDynamicsExecution.h"
+#include "RigDynamicsSolverComponent.h"
+#include "RigDynamicsParticleComponent.h"
+#include "RigDynamicsColliderComponent.h"
+
+#include "RigDynamicsExecutionHelpers.generated.h"
+
+#define UE_API CONTROLRIGDYNAMICS_API
+
+//======================================================================================================================
+USTRUCT(BlueprintType)
+struct FRigDynamicsColliderCollectionWithKey
+{
+	GENERATED_BODY()
+
+	// The bone that owns the colliders
+	UPROPERTY(BlueprintReadOnly, EditAnywhere, Category = Shapes)
+	FRigElementKey Owner = FRigElementKey(TEXT(""), ERigElementType::Bone);
+
+	// The collider properties
+	UPROPERTY(BlueprintReadOnly, EditAnywhere, Category = Shapes, meta = (ShowOnlyInnerProperties))
+	FRigDynamicsColliderCollection ColliderCollection;
+};
+
+
+//======================================================================================================================
+// Spawns a new dynamics chain (or chains)
+// Note: This node only runs as part of the construction event.
+USTRUCT(meta = (DisplayName = "Spawn Dynamics Chains", Keywords = "Add,Construction,Create,New,Simulation", Varying))
+struct FRigUnit_SpawnDynamicsChains : public FRigUnit_DynamicsBaseMutable
+{
+	GENERATED_BODY()
+
+	RIGVM_METHOD()
+	UE_API virtual void Execute() override;
+
+	FRigUnit_SpawnDynamicsChains()
+	{
+		RootBones.Add(FRigElementKey(TEXT(""), ERigElementType::Bone));
+		ParticleProperties.GravityMultiplier = 0.0f;
+	}
+
+	// An array of root bones. Each one of these will start a chain, which includes all its children/descendants
+	UPROPERTY(meta = (Input, BoneName))
+	TArray<FRigElementKey> RootBones;
+
+	// An array of bones that will terminate the chain search. Terminator bones will (if part of a
+	// chain) have particles created, but their children will not.
+	UPROPERTY(meta = (Input))
+	TArray<FRigElementKey> ChainTerminatorBones;
+
+	// Additional bones to be included
+	UPROPERTY(meta = (Input))
+	TArray<FRigElementKey> ExtraBones;
+
+	// Particle properties
+	UPROPERTY(meta = (Input))
+	FRigDynamicsParticleProperties ParticleProperties;
+
+	// Additional Colliders that track the incoming animation
+	UPROPERTY(meta = (Input))
+	TArray<FRigDynamicsColliderCollectionWithKey> Colliders;
+
+	// The element that stores the solver. If this doesn't exist, the components will be created,
+	// but not be assigned to any solver
+	UPROPERTY(meta = (Input))
+	FRigElementKey SolverElementKey = FRigElementKey(TEXT("Root"), ERigElementType::Bone);
+
+	// The name of the solver component. If this refers to a solver component that already exists,
+	// it will be used. Otherwise, a new solver component will be created, using the SolverSettings.
+	UPROPERTY(meta = (Input))
+	FName SolverComponentName = FRigDynamicsSolverComponent::GetDefaultName();
+
+	// Solver settings, if creating a solver
+	UPROPERTY(meta = (Input))
+	FRigDynamicsSolverSettings SolverSettings;
+
+	// Simulation space settings, if creating a solver
+	UPROPERTY(meta = (Input))
+	FRigDynamicsSimulationSpaceSettings SimulationSpaceSettings;
+
+	// The particle component keys that were created
+	UPROPERTY(meta = (Output))
+	TArray<FRigComponentKey> ParticleComponentKeys;
+
+	// The collider component keys that were created
+	UPROPERTY(meta = (Output))
+	TArray<FRigComponentKey> ColliderComponentKeys;
+
+	// The solver component key that was used or created
+	UPROPERTY(meta = (Output))
+	FRigComponentKey SolverComponentKey;
+};
+
+//======================================================================================================================
+// Imports collision from a physics asset and assigns it as colliders in the dynamics system
+USTRUCT(meta = (DisplayName = "Import Dynamics Colliders From Physics Asset", Keywords = "Construction,Create,New", Varying))
+struct FRigUnit_HierarchyImportDynamicsCollidersFromPhysicsAsset : public FRigUnit_DynamicsBaseMutable
+{
+	GENERATED_BODY()
+
+	RIGVM_METHOD()
+	UE_API virtual void Execute() override;
+
+	// Solver to add the colliders to
+	UPROPERTY(meta = (Input))
+	FRigComponentKey DynamicsSolverComponentKey = FRigComponentKey(
+		FRigElementKey(TEXT("Root"), ERigElementType::Bone), FRigDynamicsSolverComponent::GetDefaultName());
+
+	// The physics asset to import colliders from
+	UPROPERTY(meta = (Input))
+	TObjectPtr<UPhysicsAsset> PhysicsAsset;
+
+	// If this is empty, then all bones with bodies in the physics asset will be created. Otherwise
+	// only bodies that relate to the specified bones will be created.
+	UPROPERTY(meta = (Input))
+	TArray<FName> BoneMask;
+
+	// The dynamics collider component keys that were created
+	UPROPERTY(meta = (Output))
+	TArray<FRigComponentKey> DynamicsColliderComponentKeys;
+};
+
+
+#undef UE_API
diff --git a/Engine/Plugins/Experimental/ControlRigDynamics/Source/ControlRigDynamics/Public/RigDynamicsParticleComponent.h b/Engine/Plugins/Experimental/ControlRigDynamics/Source/ControlRigDynamics/Public/RigDynamicsParticleComponent.h
new file mode 100644
--- /dev/null
+++ b/Engine/Plugins/Experimental/ControlRigDynamics/Source/ControlRigDynamics/Public/RigDynamicsParticleComponent.h
@@ -0,0 +1,96 @@
+// Copyright Epic Games, Inc. All Rights Reserved.
+
+#pragma once
+
+#include "Rigs/RigHierarchyComponents.h"
+
+#include "RigParticleSimulation.h"
+
+#include "RigDynamicsParticleComponent.generated.h"
+
+#define UE_API CONTROLRIGDYNAMICS_API
+
+constexpr float RigDynamicsDefaultParticleRadius = 5.0f;
+constexpr float RigDynamicsDefaultParticleMass = 1.0f;
+//======================================================================================================================
+USTRUCT(BlueprintType)
+struct FRigDynamicsParticleProperties
+{
+	GENERATED_BODY()
+
+	friend FArchive& operator <<(FArchive& Ar, FRigDynamicsParticleProperties& Data);
+
+	// The radius of the particle, used for collision
+	UPROPERTY(BlueprintReadOnly, EditAnywhere, Category = Dynamics, meta = (ClampMin = "0.0"))
+	float Radius = RigDynamicsDefaultParticleRadius;
+
+	// The mass of the particle, affecting how it will respond to constraints (and external forces etc)
+	UPROPERTY(BlueprintReadOnly, EditAnywhere, Category = Dynamics, meta = (ClampMin = "0.0001"))
+	float Mass = RigDynamicsDefaultParticleMass;
+
+	// How the particle moves
+	UPROPERTY(BlueprintReadOnly, EditAnywhere, Category = Dynamics)
+	ERigParticleSimulationMovementType MovementType = ERigParticleSimulationMovementType::Simulated;
+
+	// Scales the solver gravity
+	UPROPERTY(BlueprintReadOnly, EditAnywhere, Category = Dynamics)
+	float GravityMultiplier = 1.0f;
+
+	// The strength which we drive towards the target (when simulated). This is the oscillation
+	// frequency, so low values will be soft and springy, but values significantly above
+	// 1/timestep will track the target very accurately.
+	UPROPERTY(BlueprintReadOnly, EditAnywhere, Category = Dynamics, meta = (ClampMin = "0.0"))
+	float Strength = 2.0f;
+
+	// DampingRatio for the particle
+	UPROPERTY(BlueprintReadOnly, EditAnywhere, Category = Dynamics, meta = (ClampMin = "0.0"))
+	float DampingRatio = 0.5f;
+
+	// Additional damping for the particle
+	UPROPERTY(BlueprintReadOnly, EditAnywhere, Category = Dynamics, meta = (ClampMin = "0.0"))
+	float ExtraDamping = 0.0f;
+
+	// If TargetVelocityInfluence = 1 then the target's velocity will be used for damping. If 0 then the
+	// damping is in simulation space.
+	UPROPERTY(BlueprintReadOnly, EditAnywhere, Category = Dynamics)
+	float TargetVelocityInfluence = 1.0f;
+
+	// 0 = SimSpace (child tracks absolute animation target), 1 = directional (child+parent track
+	// animation direction). Values between blend the two modes.
+	UPROPERTY(BlueprintReadOnly, EditAnywhere, Category = Dynamics, meta = (ClampMin = "0.0", ClampMax = "1.0"))
+	float TargetMode = 0.5f;
+
+	// A list of colliders we shouldn't collide with
+	UPROPERTY(BlueprintReadOnly, EditAnywhere, Category = Dynamics)
+	TArray<FRigComponentKey> NoCollisionColliders;
+
+	// A list of other particles we should collide with
+	UPROPERTY(BlueprintReadOnly, EditAnywhere, Category = Dynamics)
+	TArray<FRigComponentKey> CollisionParticles;
+};
+
+//======================================================================================================================
+// A component that can be added to a joint/element that defines how a dynamics particle can be
+// attached to it.
+//======================================================================================================================
+USTRUCT(BlueprintType)
+struct FRigDynamicsParticleComponent : public FRigBaseComponent
+{
+public:
+	GENERATED_BODY()
+	DECLARE_RIG_COMPONENT_METHODS(FRigDynamicsParticleComponent)
+
+	UPROPERTY(BlueprintReadOnly, EditAnywhere, Category = Dynamics, meta = (ShowOnlyInnerProperties))
+	FRigDynamicsParticleProperties ParticleProperties;
+
+	UE_API virtual void Save(FArchive& Ar) override;
+	UE_API virtual void Load(FArchive& Ar) override;
+
+	virtual FName GetDefaultComponentName() const override { return GetDefaultName(); }
+	static FName GetDefaultName() { return TEXT("DynamicsParticle"); }
+
+protected:
+	void Serialize(FArchive& Ar);
+};
+
+#undef UE_API
diff --git a/Engine/Plugins/Experimental/ControlRigDynamics/Source/ControlRigDynamics/Public/RigDynamicsParticleExecution.h b/Engine/Plugins/Experimental/ControlRigDynamics/Source/ControlRigDynamics/Public/RigDynamicsParticleExecution.h
new file mode 100644
--- /dev/null
+++ b/Engine/Plugins/Experimental/ControlRigDynamics/Source/ControlRigDynamics/Public/RigDynamicsParticleExecution.h
@@ -0,0 +1,452 @@
+// Copyright Epic Games, Inc. All Rights Reserved.
+
+#pragma once
+
+#include "RigDynamicsExecution.h"
+#include "RigDynamicsParticleComponent.h"
+#include "RigDynamicsColliderComponent.h"
+#include "RigDynamicsSolverComponent.h"
+
+#include "RigDynamicsParticleExecution.generated.h"
+
+#define UE_API CONTROLRIGDYNAMICS_API
+
+//======================================================================================================================
+// Spawns a new dynamics particle as a component on the owner element.
+// Note: This node only runs as part of the construction event.
+USTRUCT(meta = (DisplayName = "Spawn Dynamics Particle", Keywords = "Add,Construction,Create,New,Body", Varying))
+struct FRigUnit_SpawnDynamicsParticle : public FRigUnit_DynamicsBaseMutable
+{
+	GENERATED_BODY()
+
+	RIGVM_METHOD()
+	UE_API virtual void Execute() override;
+
+	// The owner of the newly created component (must be set/valid)
+	UPROPERTY(meta = (Input, BoneName))
+	FRigElementKey Owner = FRigElementKey(TEXT(""), ERigElementType::Bone);
+
+	// Optional solver - if set, it will be added to this solver component
+	UPROPERTY(meta = (Input))
+	FRigComponentKey DynamicsSolverComponentKey = FRigComponentKey(
+		FRigElementKey(TEXT("Root"), ERigElementType::Bone), FRigDynamicsSolverComponent::GetDefaultName());
+
+	// Particle properties
+	UPROPERTY(meta = (Input))
+	FRigDynamicsParticleProperties ParticleProperties;
+
+	// The Dynamics Particle component key that was created
+	UPROPERTY(meta = (Output))
+	FRigComponentKey DynamicsParticleComponentKey;
+};
+
+//======================================================================================================================
+// Disables collision between a particle and a collider
+USTRUCT(meta = (DisplayName = "Disable Dynamics Collision With Collider", Varying))
+struct FRigUnit_HierarchyDisableDynamicsCollisionWithCollider : public FRigUnit_DynamicsBaseMutable
+{
+	GENERATED_BODY()
+
+	RIGVM_METHOD()
+	UE_API virtual void Execute() override;
+
+	// The Dynamics Particle to be updated (which must exist)
+	UPROPERTY(meta = (Input))
+	FRigComponentKey DynamicsParticleComponentKey = FRigComponentKey(
+		FRigElementKey(TEXT(""), ERigElementType::Bone), FRigDynamicsParticleComponent::GetDefaultName());
+
+
+	// The Collider to not collide with (which doesn't have to exist yet)
+	UPROPERTY(meta = (Input))
+	FRigComponentKey DynamicsColliderComponentKey = FRigComponentKey(
+		FRigElementKey(TEXT(""), ERigElementType::Bone), FRigDynamicsColliderComponent::GetDefaultName());
+};
+
+//======================================================================================================================
+// Enables collision between a particle and another particle. The collision will actually be between
+// segments, where the segment goes from a particle top its parent (if it has one), using the radius
+// of the particles to define a capsule. TODO Note that implementation isn't yet complete, so 
+// currently it IS only sphere-sphere.
+USTRUCT(meta = (DisplayName = "Enable Dynamics Collision With Particle", Varying))
+struct FRigUnit_HierarchyEnableDynamicsCollisionWithParticle : public FRigUnit_DynamicsBaseMutable
+{
+	GENERATED_BODY()
+
+	RIGVM_METHOD()
+	UE_API virtual void Execute() override;
+
+	// The Dynamics Particle to be updated (which must exist)
+	UPROPERTY(meta = (Input))
+	FRigComponentKey DynamicsParticleComponentKey = FRigComponentKey(
+		FRigElementKey(TEXT(""), ERigElementType::Bone), FRigDynamicsParticleComponent::GetDefaultName());
+
+
+	// The other particle to collide with (which doesn't have to exist yet). Collision will be
+	// handled between segments going from each particle to its parent (if it has one). TODO Note
+	// that implementation isn't yet complete, so currently it IS only sphere-sphere.
+	UPROPERTY(meta = (Input))
+	FRigComponentKey OtherDynamicsParticleComponentKey = FRigComponentKey(
+		FRigElementKey(TEXT(""), ERigElementType::Bone), FRigDynamicsParticleComponent::GetDefaultName());
+};
+
+//======================================================================================================================
+// Sets the mass of a Dynamics particle
+USTRUCT(meta = (DisplayName = "Set Dynamics Particle Mass", Varying))
+struct FRigUnit_HierarchySetDynamicsParticleMass : public FRigUnit_DynamicsBaseMutable
+{
+	GENERATED_BODY()
+
+	RIGVM_METHOD()
+	UE_API virtual void Execute() override;
+
+	// The Dynamics Particle to be updated
+	UPROPERTY(meta = (Input))
+	FRigComponentKey DynamicsParticleComponentKey = FRigComponentKey(
+		FRigElementKey(TEXT(""), ERigElementType::Bone), FRigDynamicsParticleComponent::GetDefaultName());
+
+	// The mass of the Dynamics Particle
+	UPROPERTY(meta = (Input))
+	float Mass = RigDynamicsDefaultParticleMass;
+};
+
+//======================================================================================================================
+// Gets the mass of a Dynamics particle
+USTRUCT(meta = (DisplayName = "Get Dynamics Particle Mass", Varying))
+struct FRigUnit_HierarchyGetDynamicsParticleMass : public FRigUnit_DynamicsBase
+{
+	GENERATED_BODY()
+
+	RIGVM_METHOD()
+	UE_API virtual void Execute() override;
+
+	// The Dynamics Particle to query
+	UPROPERTY(meta = (Input))
+	FRigComponentKey DynamicsParticleComponentKey = FRigComponentKey(
+		FRigElementKey(TEXT(""), ERigElementType::Bone), FRigDynamicsParticleComponent::GetDefaultName());
+
+	// The mass of the Dynamics Particle
+	UPROPERTY(meta = (Output))
+	float Mass = RigDynamicsDefaultParticleMass;
+};
+
+//======================================================================================================================
+// Sets the radius of a Dynamics particle
+USTRUCT(meta = (DisplayName = "Set Dynamics Particle Radius", Varying))
+struct FRigUnit_HierarchySetDynamicsParticleRadius : public FRigUnit_DynamicsBaseMutable
+{
+	GENERATED_BODY()
+
+	RIGVM_METHOD()
+	UE_API virtual void Execute() override;
+
+	// The Dynamics Particle to be updated
+	UPROPERTY(meta = (Input))
+	FRigComponentKey DynamicsParticleComponentKey = FRigComponentKey(
+		FRigElementKey(TEXT(""), ERigElementType::Bone), FRigDynamicsParticleComponent::GetDefaultName());
+
+	// The radius of the Dynamics Particle, used for collision
+	UPROPERTY(meta = (Input))
+	float Radius = RigDynamicsDefaultParticleRadius;
+};
+
+//======================================================================================================================
+// Gets the radius of a Dynamics particle
+USTRUCT(meta = (DisplayName = "Get Dynamics Particle Radius", Varying))
+struct FRigUnit_HierarchyGetDynamicsParticleRadius : public FRigUnit_DynamicsBase
+{
+	GENERATED_BODY()
+
+	RIGVM_METHOD()
+	UE_API virtual void Execute() override;
+
+	// The Dynamics Particle to query
+	UPROPERTY(meta = (Input))
+	FRigComponentKey DynamicsParticleComponentKey = FRigComponentKey(
+		FRigElementKey(TEXT(""), ERigElementType::Bone), FRigDynamicsParticleComponent::GetDefaultName());
+
+	// The radius of the Dynamics Particle, used for collision
+	UPROPERTY(meta = (Output))
+	float Radius = RigDynamicsDefaultParticleRadius;
+};
+
+//======================================================================================================================
+// Sets the movement type of a Dynamics particle
+USTRUCT(meta = (DisplayName = "Set Dynamics Particle Movement Type", Varying))
+struct FRigUnit_HierarchySetDynamicsParticleMovementType : public FRigUnit_DynamicsBaseMutable
+{
+	GENERATED_BODY()
+
+	RIGVM_METHOD()
+	UE_API virtual void Execute() override;
+
+	// The Dynamics Particle to be updated
+	UPROPERTY(meta = (Input))
+	FRigComponentKey DynamicsParticleComponentKey = FRigComponentKey(
+		FRigElementKey(TEXT(""), ERigElementType::Bone), FRigDynamicsParticleComponent::GetDefaultName());
+
+	// How the particle should move
+	UPROPERTY(meta = (Input))
+	ERigParticleSimulationMovementType MovementType = ERigParticleSimulationMovementType::Simulated;
+};
+
+//======================================================================================================================
+// Gets the movement type of a Dynamics particle
+USTRUCT(meta = (DisplayName = "Get Dynamics Particle Movement Type", Varying))
+struct FRigUnit_HierarchyGetDynamicsParticleMovementType : public FRigUnit_DynamicsBase
+{
+	GENERATED_BODY()
+
+	RIGVM_METHOD()
+	UE_API virtual void Execute() override;
+
+	// The Dynamics Particle to query
+	UPROPERTY(meta = (Input))
+	FRigComponentKey DynamicsParticleComponentKey = FRigComponentKey(
+		FRigElementKey(TEXT(""), ERigElementType::Bone), FRigDynamicsParticleComponent::GetDefaultName());
+
+	// How the particle should move
+	UPROPERTY(meta = (Output))
+	ERigParticleSimulationMovementType MovementType = ERigParticleSimulationMovementType::Kinematic;
+};
+
+//======================================================================================================================
+// Sets the gravity multiplier of a Dynamics particle
+USTRUCT(meta = (DisplayName = "Set Dynamics Particle Gravity Multiplier", Varying))
+struct FRigUnit_HierarchySetDynamicsParticleGravityMultiplier : public FRigUnit_DynamicsBaseMutable
+{
+	GENERATED_BODY()
+
+	RIGVM_METHOD()
+	UE_API virtual void Execute() override;
+
+	// The Dynamics Particle to be updated
+	UPROPERTY(meta = (Input))
+	FRigComponentKey DynamicsParticleComponentKey = FRigComponentKey(
+		FRigElementKey(TEXT(""), ERigElementType::Bone), FRigDynamicsParticleComponent::GetDefaultName());
+
+	// Scales the solver gravity
+	UPROPERTY(meta = (Input))
+	float GravityMultiplier = 1.0f;
+};
+
+//======================================================================================================================
+// Gets the gravity multiplier of a Dynamics particle
+USTRUCT(meta = (DisplayName = "Get Dynamics Particle Gravity Multiplier", Varying))
+struct FRigUnit_HierarchyGetDynamicsParticleGravityMultiplier : public FRigUnit_DynamicsBase
+{
+	GENERATED_BODY()
+
+	RIGVM_METHOD()
+	UE_API virtual void Execute() override;
+
+	// The Dynamics Particle to query
+	UPROPERTY(meta = (Input))
+	FRigComponentKey DynamicsParticleComponentKey = FRigComponentKey(
+		FRigElementKey(TEXT(""), ERigElementType::Bone), FRigDynamicsParticleComponent::GetDefaultName());
+
+	// Scales the solver gravity
+	UPROPERTY(meta = (Output))
+	float GravityMultiplier = 0.0f;
+};
+
+//======================================================================================================================
+// Sets the strength of a Dynamics particle
+USTRUCT(meta = (DisplayName = "Set Dynamics Particle Strength", Varying))
+struct FRigUnit_HierarchySetDynamicsParticleStrength : public FRigUnit_DynamicsBaseMutable
+{
+	GENERATED_BODY()
+
+	RIGVM_METHOD()
+	UE_API virtual void Execute() override;
+
+	// The Dynamics Particle to be updated
+	UPROPERTY(meta = (Input))
+	FRigComponentKey DynamicsParticleComponentKey = FRigComponentKey(
+		FRigElementKey(TEXT(""), ERigElementType::Bone), FRigDynamicsParticleComponent::GetDefaultName());
+
+	// The strength which drives towards the target (oscillation frequency)
+	UPROPERTY(meta = (Input, ClampMin = "0.0"))
+	float Strength = 2.0;
+};
+
+//======================================================================================================================
+// Gets the strength of a Dynamics particle
+USTRUCT(meta = (DisplayName = "Get Dynamics Particle Strength", Varying))
+struct FRigUnit_HierarchyGetDynamicsParticleStrength : public FRigUnit_DynamicsBase
+{
+	GENERATED_BODY()
+
+	RIGVM_METHOD()
+	UE_API virtual void Execute() override;
+
+	// The Dynamics Particle to query
+	UPROPERTY(meta = (Input))
+	FRigComponentKey DynamicsParticleComponentKey = FRigComponentKey(
+		FRigElementKey(TEXT(""), ERigElementType::Bone), FRigDynamicsParticleComponent::GetDefaultName());
+
+	// The strength which drives towards the target (oscillation frequency)
+	UPROPERTY(meta = (Output))
+	float Strength = 0.0f;
+};
+
+//======================================================================================================================
+// Sets the damping ratio of a Dynamics particle
+USTRUCT(meta = (DisplayName = "Set Dynamics Particle Damping Ratio", Varying))
+struct FRigUnit_HierarchySetDynamicsParticleDampingRatio : public FRigUnit_DynamicsBaseMutable
+{
+	GENERATED_BODY()
+
+	RIGVM_METHOD()
+	UE_API virtual void Execute() override;
+
+	// The Dynamics Particle to be updated
+	UPROPERTY(meta = (Input))
+	FRigComponentKey DynamicsParticleComponentKey = FRigComponentKey(
+		FRigElementKey(TEXT(""), ERigElementType::Bone), FRigDynamicsParticleComponent::GetDefaultName());
+
+	// Damping ratio for the particle
+	UPROPERTY(meta = (Input, ClampMin = "0.0"))
+	float DampingRatio = 1.0;
+};
+
+//======================================================================================================================
+// Gets the damping ratio of a Dynamics particle
+USTRUCT(meta = (DisplayName = "Get Dynamics Particle Damping Ratio", Varying))
+struct FRigUnit_HierarchyGetDynamicsParticleDampingRatio : public FRigUnit_DynamicsBase
+{
+	GENERATED_BODY()
+
+	RIGVM_METHOD()
+	UE_API virtual void Execute() override;
+
+	// The Dynamics Particle to query
+	UPROPERTY(meta = (Input))
+	FRigComponentKey DynamicsParticleComponentKey = FRigComponentKey(
+		FRigElementKey(TEXT(""), ERigElementType::Bone), FRigDynamicsParticleComponent::GetDefaultName());
+
+	// Damping ratio for the particle
+	UPROPERTY(meta = (Output))
+	float DampingRatio = 0.0f;
+};
+
+//======================================================================================================================
+// Sets the extra damping of a Dynamics particle
+USTRUCT(meta = (DisplayName = "Set Dynamics Particle Extra Damping", Varying))
+struct FRigUnit_HierarchySetDynamicsParticleExtraDamping : public FRigUnit_DynamicsBaseMutable
+{
+	GENERATED_BODY()
+
+	RIGVM_METHOD()
+	UE_API virtual void Execute() override;
+
+	// The Dynamics Particle to be updated
+	UPROPERTY(meta = (Input))
+	FRigComponentKey DynamicsParticleComponentKey = FRigComponentKey(
+		FRigElementKey(TEXT(""), ERigElementType::Bone), FRigDynamicsParticleComponent::GetDefaultName());
+
+	// Extra damping for the particle
+	UPROPERTY(meta = (Input, ClampMin = "0.0"))
+	float ExtraDamping = 0.0;
+};
+
+//======================================================================================================================
+// Gets the extra damping of a Dynamics particle
+USTRUCT(meta = (DisplayName = "Get Dynamics Particle Extra Damping", Varying))
+struct FRigUnit_HierarchyGetDynamicsParticleExtraDamping : public FRigUnit_DynamicsBase
+{
+	GENERATED_BODY()
+
+	RIGVM_METHOD()
+	UE_API virtual void Execute() override;
+
+	// The Dynamics Particle to query
+	UPROPERTY(meta = (Input))
+	FRigComponentKey DynamicsParticleComponentKey = FRigComponentKey(
+		FRigElementKey(TEXT(""), ERigElementType::Bone), FRigDynamicsParticleComponent::GetDefaultName());
+
+	// Extra damping for the particle
+	UPROPERTY(meta = (Output))
+	float ExtraDamping = 0.0f;
+};
+
+//======================================================================================================================
+// Sets the target velocity influence of a Dynamics particle
+USTRUCT(meta = (DisplayName = "Set Dynamics Particle Target Velocity Influence", Varying))
+struct FRigUnit_HierarchySetDynamicsParticleTargetVelocityInfluence : public FRigUnit_DynamicsBaseMutable
+{
+	GENERATED_BODY()
+
+	RIGVM_METHOD()
+	UE_API virtual void Execute() override;
+
+	// The Dynamics Particle to be updated
+	UPROPERTY(meta = (Input))
+	FRigComponentKey DynamicsParticleComponentKey = FRigComponentKey(
+		FRigElementKey(TEXT(""), ERigElementType::Bone), FRigDynamicsParticleComponent::GetDefaultName());
+
+	// If 1 the target's velocity will be used for damping, if 0 damping is in sim space
+	UPROPERTY(meta = (Input, ClampMin = "0.0"))
+	float TargetVelocityInfluence = 1.0;
+};
+
+//======================================================================================================================
+// Gets the target velocity influence of a Dynamics particle
+USTRUCT(meta = (DisplayName = "Get Dynamics Particle Target Velocity Influence", Varying))
+struct FRigUnit_HierarchyGetDynamicsParticleTargetVelocityInfluence : public FRigUnit_DynamicsBase
+{
+	GENERATED_BODY()
+
+	RIGVM_METHOD()
+	UE_API virtual void Execute() override;
+
+	// The Dynamics Particle to query
+	UPROPERTY(meta = (Input))
+	FRigComponentKey DynamicsParticleComponentKey = FRigComponentKey(
+		FRigElementKey(TEXT(""), ERigElementType::Bone), FRigDynamicsParticleComponent::GetDefaultName());
+
+	// If 1 the target's velocity will be used for damping, if 0 damping is in sim space
+	UPROPERTY(meta = (Output))
+	float TargetVelocityInfluence = 0.0f;
+};
+
+//======================================================================================================================
+// Sets the target mode of a Dynamics particle
+USTRUCT(meta = (DisplayName = "Set Dynamics Particle Target Mode", Varying))
+struct FRigUnit_HierarchySetDynamicsParticleTargetMode : public FRigUnit_DynamicsBaseMutable
+{
+	GENERATED_BODY()
+
+	RIGVM_METHOD()
+	UE_API virtual void Execute() override;
+
+	// The Dynamics Particle to be updated
+	UPROPERTY(meta = (Input))
+	FRigComponentKey DynamicsParticleComponentKey = FRigComponentKey(
+		FRigElementKey(TEXT(""), ERigElementType::Bone), FRigDynamicsParticleComponent::GetDefaultName());
+
+	// 0 = SimSpace, 1 = directional, values between blend the two modes
+	UPROPERTY(meta = (Input, ClampMin = "0.0", ClampMax = "1.0"))
+	float TargetMode = 0.5f;
+};
+
+//======================================================================================================================
+// Gets the target mode of a Dynamics particle
+USTRUCT(meta = (DisplayName = "Get Dynamics Particle Target Mode", Varying))
+struct FRigUnit_HierarchyGetDynamicsParticleTargetMode : public FRigUnit_DynamicsBase
+{
+	GENERATED_BODY()
+
+	RIGVM_METHOD()
+	UE_API virtual void Execute() override;
+
+	// The Dynamics Particle to query
+	UPROPERTY(meta = (Input))
+	FRigComponentKey DynamicsParticleComponentKey = FRigComponentKey(
+		FRigElementKey(TEXT(""), ERigElementType::Bone), FRigDynamicsParticleComponent::GetDefaultName());
+
+	// 0 = SimSpace, 1 = directional, values between blend the two modes
+	UPROPERTY(meta = (Output))
+	float TargetMode = 0.5f;
+};
+
+#undef UE_API
diff --git a/Engine/Plugins/Experimental/ControlRigDynamics/Source/ControlRigDynamics/Public/RigDynamicsSimulationSpace.h b/Engine/Plugins/Experimental/ControlRigDynamics/Source/ControlRigDynamics/Public/RigDynamicsSimulationSpace.h
new file mode 100644
--- /dev/null
+++ b/Engine/Plugins/Experimental/ControlRigDynamics/Source/ControlRigDynamics/Public/RigDynamicsSimulationSpace.h
@@ -0,0 +1,115 @@
+// Copyright Epic Games, Inc. All Rights Reserved.
+
+#pragma once
+
+#include "CoreMinimal.h"
+
+#include "RigDynamicsSolverComponent.h"
+#include "RigParticleSimulation.h"
+
+#define UE_API CONTROLRIGDYNAMICS_API
+
+//======================================================================================================================
+// This represents the simulation space, specified in world space. It will be associated with
+// functions that allow conversion between the spaces. 
+struct FRigDynamicsSimulationSpaceState
+{
+public:
+	// If DeltaTime > 0 then this does a full update, calculating new velocities etc. If DeltaTime
+	// <= 0 then this just updates the current position and orientation - no changes to velocity etc.
+	void Update(const FTransform& InComponentTM, const FTransform& InSimulationSpaceTM, float DeltaTime);
+
+	// This sets the position/orientation, and resets everything else.
+	void Reset(const FTransform& InComponentTM, const FTransform& InSimulationSpaceTM);
+
+	// Calculates the simulation space motion data from the current state and settings, ready
+	// to be passed to RigParticleSimulation::Simulate.
+	RigParticleSimulation::FSimulationSpaceMotion CalculateMotion(
+		const FRigDynamicsSimulationSpaceSettings& Settings, double AbsoluteTime) const;
+
+	// Convert a position from component space (where the hierarchy lives) into simulation
+	// space (where particles live). Identity when SimulationSpace=Component.
+	FVector ConvertComponentSpacePositionToSimSpace(const FVector& ComponentSpacePosition) const;
+
+	// Convert a position from simulation space back to component space.
+	FVector ConvertSimSpacePositionToComponentSpace(const FVector& SimSpacePosition) const;
+
+	// Convert a transform from component space into simulation space (rotation + translation, no scale).
+	FTransform ConvertComponentSpaceTransformToSimSpace(const FTransform& ComponentSpaceTransform) const;
+
+	// Convert a transform from simulation space back to component space.
+	FTransform ConvertSimSpaceTransformToComponentSpace(const FTransform& SimSpaceTransform) const;
+
+	// Convert a direction/velocity vector from world space into simulation space.
+	FVector ConvertWorldVectorToSimSpace(const FVector& WorldVector) const;
+
+	// The world-space component (owner of the simulation) transform
+	FTransform ComponentTM;
+
+	// World-space linear velocity of the simulation space
+	FVector LinearVelocity = FVector::ZeroVector;
+	// World-space angular velocity of the simulation space
+	FVector AngularVelocity = FVector::ZeroVector; // rad/s
+
+	// World-space linear acceleration of the simulation space
+	FVector LinearAcceleration = FVector::ZeroVector;
+	// World-space angular acceleration of the simulation space
+	FVector AngularAcceleration = FVector::ZeroVector; // rad/s/s
+
+	// The world-space simulation space TM
+	FTransform SimulationSpaceTM;
+	FTransform PrevSimulationSpaceTM;
+	FTransform PrevPrevSimulationSpaceTM;
+
+	// The time between SimulationSpaceTM and PrevSimulationSpaceTM
+	float DeltaTime = 0.0f;
+	// The time between PrevSimulationSpaceTM and PrevPrevSimulationSpaceTM
+	float PrevDeltaTime = 0.0f;
+
+	// When we are reset, then the position (and orientation) is correct
+	// On the next update (with a non-zero delta time), the velocity is correct (acceleration will be zero)
+	// On the next update, the acceleration will be correct
+	int64 UpdatesSinceReset = 0;
+};
+
+//======================================================================================================================
+// Inline implementations
+//======================================================================================================================
+
+inline FVector FRigDynamicsSimulationSpaceState::ConvertComponentSpacePositionToSimSpace(
+	const FVector& ComponentSpacePosition) const
+{
+	return SimulationSpaceTM.InverseTransformPositionNoScale(
+		ComponentTM.TransformPositionNoScale(ComponentSpacePosition));
+}
+
+inline FVector FRigDynamicsSimulationSpaceState::ConvertSimSpacePositionToComponentSpace(
+	const FVector& SimSpacePosition) const
+{
+	return ComponentTM.InverseTransformPositionNoScale(
+		SimulationSpaceTM.TransformPositionNoScale(SimSpacePosition));
+}
+
+inline FTransform FRigDynamicsSimulationSpaceState::ConvertComponentSpaceTransformToSimSpace(
+	const FTransform& ComponentSpaceTransform) const
+{
+	FTransform Result = ComponentSpaceTransform * ComponentTM * SimulationSpaceTM.Inverse();
+	Result.SetScale3D(FVector::OneVector);
+	return Result;
+}
+
+inline FTransform FRigDynamicsSimulationSpaceState::ConvertSimSpaceTransformToComponentSpace(
+	const FTransform& SimSpaceTransform) const
+{
+	FTransform Result = SimSpaceTransform * SimulationSpaceTM * ComponentTM.Inverse();
+	Result.SetScale3D(FVector::OneVector);
+	return Result;
+}
+
+inline FVector FRigDynamicsSimulationSpaceState::ConvertWorldVectorToSimSpace(
+	const FVector& WorldVector) const
+{
+	return SimulationSpaceTM.InverseTransformVectorNoScale(WorldVector);
+}
+
+#undef UE_API
diff --git a/Engine/Plugins/Experimental/ControlRigDynamics/Source/ControlRigDynamics/Public/RigDynamicsSolverComponent.h b/Engine/Plugins/Experimental/ControlRigDynamics/Source/ControlRigDynamics/Public/RigDynamicsSolverComponent.h
new file mode 100644
--- /dev/null
+++ b/Engine/Plugins/Experimental/ControlRigDynamics/Source/ControlRigDynamics/Public/RigDynamicsSolverComponent.h
@@ -0,0 +1,74 @@
+// Copyright Epic Games, Inc. All Rights Reserved.
+
+#pragma once
+
+#include "RigDynamicsData.h"
+
+#include "Rigs/RigHierarchyComponents.h"
+
+#include "Templates/SharedPointer.h"
+
+#include "RigDynamicsSolverComponent.generated.h"
+
+#define UE_API CONTROLRIGDYNAMICS_API
+
+struct FRigDynamicsSolver;
+
+//======================================================================================================================
+// A component that can be added to a joint/element that defines how a system of particles can be solved
+//======================================================================================================================
+USTRUCT(BlueprintType)
+struct FRigDynamicsSolverComponent : public FRigBaseComponent
+{
+public:
+	GENERATED_BODY()
+	DECLARE_RIG_COMPONENT_METHODS(FRigDynamicsSolverComponent)
+
+	// Particles can be dynamic or kinematic. Kinematic particles track the target bones. Dynamic
+	// particles fall under gravity, respond to constraints etc. In addition, dynamic particles will
+	// track the target bones using forces. They will also be connected in chains to maintain the
+	// original bone lengths.
+	UPROPERTY(BlueprintReadOnly, EditAnywhere, Category = Components)
+	TArray<FRigComponentKey> Particles;
+
+	// Collision shapes attached to bones in the hierarchy
+	UPROPERTY(BlueprintReadOnly, EditAnywhere, Category = Components)
+	TArray<FRigComponentKey> Colliders;
+
+	// Constraints apply forces to pairs of particles to maintain a target distance
+	UPROPERTY(BlueprintReadOnly, EditAnywhere, Category = Components)
+	TArray<FRigComponentKey> Constraints;
+
+	UPROPERTY(BlueprintReadOnly, EditAnywhere, Category = Solver)
+	FRigDynamicsSolverSettings Settings;
+
+	UPROPERTY(BlueprintReadOnly, EditAnywhere, Category = Solver)
+	FRigDynamicsSimulationSpaceSettings SimulationSpaceSettings;
+
+	UE_API virtual void Save(FArchive& Ar) override;
+	UE_API virtual void Load(FArchive& Ar) override;
+
+	virtual FName GetDefaultComponentName() const override { return GetDefaultName(); }
+	static FName GetDefaultName() { return TEXT("DynamicsSolver"); }
+
+	// This will make the internally owned solver if necessary, and return it
+	FRigDynamicsSolver* GetDynamicsSolver() const;
+
+protected:
+	void Serialize(FArchive& Ar);
+
+private:
+
+	// We use a TSharedPtr here in anticipation of supporting world collisions (see Control Rig Physics).
+	//
+	// We make it mutable so that GetDynamicsSolver can be const - in practice the user shouldn't
+	// have to worry.
+	//
+	// Finally, UE's reflection system wants to be able to use the assignment operator. If it wasn't
+	// for that, we could just delete the copy operator etc and prevent duplicates containing a
+	// pointer to the same DynamicsSolver. In practice, this won't ever happen for a "live"
+	// component, so we're OK.
+	mutable TSharedPtr<FRigDynamicsSolver> DynamicsSolverPtr;
+};
+
+#undef UE_API
\ No newline at end of file
diff --git a/Engine/Plugins/Experimental/ControlRigDynamics/Source/ControlRigDynamics/Public/RigDynamicsSolverExecution.h b/Engine/Plugins/Experimental/ControlRigDynamics/Source/ControlRigDynamics/Public/RigDynamicsSolverExecution.h
new file mode 100644
--- /dev/null
+++ b/Engine/Plugins/Experimental/ControlRigDynamics/Source/ControlRigDynamics/Public/RigDynamicsSolverExecution.h
@@ -0,0 +1,345 @@
+// Copyright Epic Games, Inc. All Rights Reserved.
+
+#pragma once
+
+#include "RigDynamicsExecution.h"
+#include "RigDynamicsSolverComponent.h"
+
+#include "Rigs/RigHierarchyCache.h"
+
+#include "RigDynamicsSolverExecution.generated.h"
+
+#define UE_API CONTROLRIGDYNAMICS_API
+
+//======================================================================================================================
+// Spawns a new dynamics solver as a component on the owner element.
+// Note: This node only runs as part of the construction event.
+USTRUCT(meta = (DisplayName = "Spawn Dynamics Solver", Keywords = "Add,Construction,Create,New,Simulation", Varying))
+struct FRigUnit_SpawnDynamicsSolver : public FRigUnit_DynamicsBaseMutable
+{
+	GENERATED_BODY()
+
+	RIGVM_METHOD()
+	UE_API virtual void Execute() override;
+
+	// The owner of the newly created component (must be set/valid)
+	UPROPERTY(meta = (Input, BoneName))
+	FRigElementKey Owner = FRigElementKey(FName(TEXT("Root")), ERigElementType::Bone);
+
+	// Particles to be included in the dynamics simulation (more can be added during construction)
+	UPROPERTY(meta = (Input))
+	TArray<FRigComponentKey> Particles;
+
+	// Colliders to be included in the dynamics simulation (more can be added during construction)
+	UPROPERTY(meta = (Input))
+	TArray<FRigComponentKey> Colliders;
+
+	// Solver settings
+	UPROPERTY(meta = (Input))
+	FRigDynamicsSolverSettings Settings;
+
+	// Simulation space settings
+	UPROPERTY(meta = (Input))
+	FRigDynamicsSimulationSpaceSettings SimulationSpaceSettings;
+
+	// The solver component key that was created
+	UPROPERTY(meta = (Output))
+	FRigComponentKey DynamicsSolverComponentKey;
+};
+
+//======================================================================================================================
+// Instantiates all the objects in the dynamics world. Some properties can't be modified after this happens.
+// Note that it will happen automatically during the first simulation step if it hasn't been explicitly
+// requested. Explicit instantiation allows the timing to be controlled, as allocations etc may cause some
+// delays.
+USTRUCT(meta = (DisplayName = "Instantiate dynamics"))
+struct FRigUnit_InstantiateDynamics : public FRigUnit_DynamicsBaseMutable
+{
+	GENERATED_BODY()
+
+	RIGVM_METHOD()
+	UE_API virtual void Execute() override;
+
+	// The solver that should be instantiated
+	UPROPERTY(meta = (Input))
+	FRigComponentKey DynamicsSolverComponentKey = FRigComponentKey(
+		FRigElementKey(TEXT("Root"), ERigElementType::Bone), FRigDynamicsSolverComponent::GetDefaultName());
+};
+
+
+//======================================================================================================================
+// Steps the dynamics solver
+USTRUCT(meta = (DisplayName = "Step Dynamics Solver", Keywords = "Simulate"))
+struct FRigUnit_StepDynamicsSolver : public FRigUnit_DynamicsBaseMutable
+{
+	GENERATED_BODY()
+
+	RIGVM_METHOD()
+	UE_API virtual void Execute() override;
+
+	// The physics solver that should be stepped
+	UPROPERTY(meta = (Input))
+	FRigComponentKey DynamicsSolverComponentKey = FRigComponentKey(
+		FRigElementKey(TEXT("Root"), ERigElementType::Bone), FRigDynamicsSolverComponent::GetDefaultName());
+
+	// If this is zero, then the execute context time will be used. If this is positive then it will
+	// override the delta time. A negative value will prevent the solver from stepping.
+	UPROPERTY(meta = (Input))
+	float DeltaTimeOverride = 0.0f;
+
+	// How much of the simulation is combined with the input bone. This currently happens in
+	// component space. Note that the simulation will continue to run, even if alpha = 0
+	UPROPERTY(meta = (Input, ClampMin = "0.0", ClampMax = "1.0"))
+	float Alpha = 1.0f;
+
+	// Settings that specify how the solver state should be visualized during/after the step
+	UPROPERTY(meta = (Input))
+	FRigDynamicsVisualizationSettings VisualizationSettings;
+
+	// Doing a lookup to get the SolverComponent from the key is slow, so cache this. If it
+	// matches DynamicsSolverComponentKey we can use the cached value.
+	UPROPERTY()
+	FCachedRigComponent CachedSolverComponent;
+};
+
+//======================================================================================================================
+// Sets the complete settings of a Dynamics solver
+USTRUCT(meta = (DisplayName = "Set Dynamics Solver Settings", Varying))
+struct FRigUnit_HierarchySetDynamicsSolverSettings : public FRigUnit_DynamicsBaseMutable
+{
+	GENERATED_BODY()
+
+	RIGVM_METHOD()
+	UE_API virtual void Execute() override;
+
+	// The Dynamics Solver to be updated
+	UPROPERTY(meta = (Input))
+	FRigComponentKey DynamicsSolverComponentKey = FRigComponentKey(
+		FRigElementKey(TEXT("Root"), ERigElementType::Bone), FRigDynamicsSolverComponent::GetDefaultName());
+
+	// The solver settings
+	UPROPERTY(meta = (Input))
+	FRigDynamicsSolverSettings Settings;
+};
+
+//======================================================================================================================
+// Gets the complete settings of a Dynamics solver
+USTRUCT(meta = (DisplayName = "Get Dynamics Solver Settings", Varying))
+struct FRigUnit_HierarchyGetDynamicsSolverSettings : public FRigUnit_DynamicsBase
+{
+	GENERATED_BODY()
+
+	RIGVM_METHOD()
+	UE_API virtual void Execute() override;
+
+	// The Dynamics Solver to query
+	UPROPERTY(meta = (Input))
+	FRigComponentKey DynamicsSolverComponentKey = FRigComponentKey(
+		FRigElementKey(TEXT("Root"), ERigElementType::Bone), FRigDynamicsSolverComponent::GetDefaultName());
+
+	// The solver settings
+	UPROPERTY(meta = (Output))
+	FRigDynamicsSolverSettings Settings;
+};
+
+//======================================================================================================================
+// Sets the gravity of a Dynamics solver
+USTRUCT(meta = (DisplayName = "Set Dynamics Solver Gravity", Varying))
+struct FRigUnit_HierarchySetDynamicsSolverGravity : public FRigUnit_DynamicsBaseMutable
+{
+	GENERATED_BODY()
+
+	RIGVM_METHOD()
+	UE_API virtual void Execute() override;
+
+	// The Dynamics Solver to be updated
+	UPROPERTY(meta = (Input))
+	FRigComponentKey DynamicsSolverComponentKey = FRigComponentKey(
+		FRigElementKey(TEXT("Root"), ERigElementType::Bone), FRigDynamicsSolverComponent::GetDefaultName());
+
+	// Gravity in world space, can be scaled per particle
+	UPROPERTY(meta = (Input))
+	FVector Gravity = FVector(0, 0, -981.0);
+};
+
+//======================================================================================================================
+// Gets the gravity of a Dynamics solver
+USTRUCT(meta = (DisplayName = "Get Dynamics Solver Gravity", Varying))
+struct FRigUnit_HierarchyGetDynamicsSolverGravity : public FRigUnit_DynamicsBase
+{
+	GENERATED_BODY()
+
+	RIGVM_METHOD()
+	UE_API virtual void Execute() override;
+
+	// The Dynamics Solver to query
+	UPROPERTY(meta = (Input))
+	FRigComponentKey DynamicsSolverComponentKey = FRigComponentKey(
+		FRigElementKey(TEXT("Root"), ERigElementType::Bone), FRigDynamicsSolverComponent::GetDefaultName());
+
+	// Gravity in world space, can be scaled per particle
+	UPROPERTY(meta = (Output))
+	FVector Gravity = FVector::ZeroVector;
+};
+
+//======================================================================================================================
+// Sets the max time step of a Dynamics solver
+USTRUCT(meta = (DisplayName = "Set Dynamics Solver Max Time Step", Varying))
+struct FRigUnit_HierarchySetDynamicsSolverMaxTimeStep : public FRigUnit_DynamicsBaseMutable
+{
+	GENERATED_BODY()
+
+	RIGVM_METHOD()
+	UE_API virtual void Execute() override;
+
+	// The Dynamics Solver to be updated
+	UPROPERTY(meta = (Input))
+	FRigComponentKey DynamicsSolverComponentKey = FRigComponentKey(
+		FRigElementKey(TEXT("Root"), ERigElementType::Bone), FRigDynamicsSolverComponent::GetDefaultName());
+
+	// The maximum timestep of any step
+	UPROPERTY(meta = (Input, ClampMin = "0.0"))
+	float MaxTimeStep = 1.0f;
+};
+
+//======================================================================================================================
+// Gets the max time step of a Dynamics solver
+USTRUCT(meta = (DisplayName = "Get Dynamics Solver Max Time Step", Varying))
+struct FRigUnit_HierarchyGetDynamicsSolverMaxTimeStep : public FRigUnit_DynamicsBase
+{
+	GENERATED_BODY()
+
+	RIGVM_METHOD()
+	UE_API virtual void Execute() override;
+
+	// The Dynamics Solver to query
+	UPROPERTY(meta = (Input))
+	FRigComponentKey DynamicsSolverComponentKey = FRigComponentKey(
+		FRigElementKey(TEXT("Root"), ERigElementType::Bone), FRigDynamicsSolverComponent::GetDefaultName());
+
+	// The maximum timestep of any step
+	UPROPERTY(meta = (Output))
+	float MaxTimeStep = 0.0f;
+};
+
+//======================================================================================================================
+// Sets the max number of steps of a Dynamics solver
+USTRUCT(meta = (DisplayName = "Set Dynamics Solver Max Num Steps", Varying))
+struct FRigUnit_HierarchySetDynamicsSolverMaxNumSteps : public FRigUnit_DynamicsBaseMutable
+{
+	GENERATED_BODY()
+
+	RIGVM_METHOD()
+	UE_API virtual void Execute() override;
+
+	// The Dynamics Solver to be updated
+	UPROPERTY(meta = (Input))
+	FRigComponentKey DynamicsSolverComponentKey = FRigComponentKey(
+		FRigElementKey(TEXT("Root"), ERigElementType::Bone), FRigDynamicsSolverComponent::GetDefaultName());
+
+	// The maximum number of steps per update
+	UPROPERTY(meta = (Input, ClampMin = "0"))
+	int32 MaxNumSteps = 1;
+};
+
+//======================================================================================================================
+// Gets the max number of steps of a Dynamics solver
+USTRUCT(meta = (DisplayName = "Get Dynamics Solver Max Num Steps", Varying))
+struct FRigUnit_HierarchyGetDynamicsSolverMaxNumSteps : public FRigUnit_DynamicsBase
+{
+	GENERATED_BODY()
+
+	RIGVM_METHOD()
+	UE_API virtual void Execute() override;
+
+	// The Dynamics Solver to query
+	UPROPERTY(meta = (Input))
+	FRigComponentKey DynamicsSolverComponentKey = FRigComponentKey(
+		FRigElementKey(TEXT("Root"), ERigElementType::Bone), FRigDynamicsSolverComponent::GetDefaultName());
+
+	// The maximum number of steps per update
+	UPROPERTY(meta = (Output))
+	int32 MaxNumSteps = 0;
+};
+
+//======================================================================================================================
+// Sets the number of iterations of a Dynamics solver
+USTRUCT(meta = (DisplayName = "Set Dynamics Solver Num Iterations", Varying))
+struct FRigUnit_HierarchySetDynamicsSolverNumIterations : public FRigUnit_DynamicsBaseMutable
+{
+	GENERATED_BODY()
+
+	RIGVM_METHOD()
+	UE_API virtual void Execute() override;
+
+	// The Dynamics Solver to be updated
+	UPROPERTY(meta = (Input))
+	FRigComponentKey DynamicsSolverComponentKey = FRigComponentKey(
+		FRigElementKey(TEXT("Root"), ERigElementType::Bone), FRigDynamicsSolverComponent::GetDefaultName());
+
+	// Total number of iterations including particle target tracking
+	UPROPERTY(meta = (Input, ClampMin = "0"))
+	int32 NumIterations = 1;
+};
+
+//======================================================================================================================
+// Gets the number of iterations of a Dynamics solver
+USTRUCT(meta = (DisplayName = "Get Dynamics Solver Num Iterations", Varying))
+struct FRigUnit_HierarchyGetDynamicsSolverNumIterations : public FRigUnit_DynamicsBase
+{
+	GENERATED_BODY()
+
+	RIGVM_METHOD()
+	UE_API virtual void Execute() override;
+
+	// The Dynamics Solver to query
+	UPROPERTY(meta = (Input))
+	FRigComponentKey DynamicsSolverComponentKey = FRigComponentKey(
+		FRigElementKey(TEXT("Root"), ERigElementType::Bone), FRigDynamicsSolverComponent::GetDefaultName());
+
+	// Total number of iterations including particle target tracking
+	UPROPERTY(meta = (Output))
+	int32 NumIterations = 0;
+};
+
+//======================================================================================================================
+// Sets the number of constraint sub-iterations of a Dynamics solver
+USTRUCT(meta = (DisplayName = "Set Dynamics Solver Num Constraint Sub-Iterations", Varying))
+struct FRigUnit_HierarchySetDynamicsSolverNumConstraintSubIterations : public FRigUnit_DynamicsBaseMutable
+{
+	GENERATED_BODY()
+
+	RIGVM_METHOD()
+	UE_API virtual void Execute() override;
+
+	// The Dynamics Solver to be updated
+	UPROPERTY(meta = (Input))
+	FRigComponentKey DynamicsSolverComponentKey = FRigComponentKey(
+		FRigElementKey(TEXT("Root"), ERigElementType::Bone), FRigDynamicsSolverComponent::GetDefaultName());
+
+	// Number of constraint (distance and collision) sub-iterations
+	UPROPERTY(meta = (Input, ClampMin = "0"))
+	int32 NumConstraintSubIterations = 1;
+};
+
+//======================================================================================================================
+// Gets the number of constraint sub-iterations of a Dynamics solver
+USTRUCT(meta = (DisplayName = "Get Dynamics Solver Num Constraint Sub-Iterations", Varying))
+struct FRigUnit_HierarchyGetDynamicsSolverNumConstraintSubIterations : public FRigUnit_DynamicsBase
+{
+	GENERATED_BODY()
+
+	RIGVM_METHOD()
+	UE_API virtual void Execute() override;
+
+	// The Dynamics Solver to query
+	UPROPERTY(meta = (Input))
+	FRigComponentKey DynamicsSolverComponentKey = FRigComponentKey(
+		FRigElementKey(TEXT("Root"), ERigElementType::Bone), FRigDynamicsSolverComponent::GetDefaultName());
+
+	// Number of constraint (distance and collision) sub-iterations
+	UPROPERTY(meta = (Output))
+	int32 NumConstraintSubIterations = 0;
+};
+
+#undef UE_API
diff --git a/Engine/Plugins/Experimental/ControlRigDynamics/Source/ControlRigDynamics/Public/RigParticleSimulation.h b/Engine/Plugins/Experimental/ControlRigDynamics/Source/ControlRigDynamics/Public/RigParticleSimulation.h
new file mode 100644
--- /dev/null
+++ b/Engine/Plugins/Experimental/ControlRigDynamics/Source/ControlRigDynamics/Public/RigParticleSimulation.h
@@ -0,0 +1,502 @@
+// Copyright Epic Games, Inc. All Rights Reserved.
+
+#pragma once
+
+#include "CoreMinimal.h"
+
+#include "RigParticleSimulation.generated.h"
+
+#define UE_API CONTROLRIGDYNAMICS_API
+
+// Uncomment to use single-precision (float) vectors/transforms in the simulation. By default the
+// simulation uses the engine's FVector/FTransform (double).
+//
+// However, note that profiling indicates only the tiniest improvement through using floats, even
+// when multiple iterations are used.
+
+// #define CONTROLRIGDYNAMICS_SINGLEPRECISION
+
+//======================================================================================================================
+// This should be kept completely independent of control rig etc
+//======================================================================================================================
+
+//======================================================================================================================
+UENUM(BlueprintType)
+enum class ERigParticleSimulationMovementType : uint8
+{
+	// Kinematic means that the particle will track animation
+	Kinematic,
+	// Simulated means that the particle will be controlled by dynamics
+	Simulated
+};
+
+//======================================================================================================================
+// TODO support min, max and ranged length constraints
+UENUM(BlueprintType)
+enum class ERigParticleSimulationConstraintType : uint8
+{
+	// The target distance will be enforced as hard as possible
+	Hard,
+	// The target distance will be enforced softly, based on the strengths
+	Soft
+};
+
+namespace RigParticleSimulation
+{
+
+#ifdef CONTROLRIGDYNAMICS_SINGLEPRECISION
+using FSimVector    = FVector3f;
+using FSimTransform = FTransform3f;
+using FSimVector2D  = FVector2f;
+using FSimScalar    = float;
+#else
+using FSimVector    = FVector;
+using FSimTransform = FTransform;
+using FSimVector2D  = FVector2D;
+using FSimScalar    = double;
+#endif
+
+//======================================================================================================================
+// Simulation data used in the central loops of the solve.
+// 
+// Note that each one of these has an Update function, which may need to be called before the
+// simulate call. The user should only modify things through this update (after instantiation).
+//======================================================================================================================
+
+//======================================================================================================================
+// We have an "Info" structure that accompanies the simulation structure. The Info contains data
+// that is "user-facing" (the user is control rig), and data that is helpful for updating the
+// simulation structures.
+struct FParticleInfo
+{
+	// Index into the Particles/ParticleInfos arrays to retrieve the parent particle
+	int32 ParentParticleIndex = INDEX_NONE;
+	int32 NumChildren = 0;
+
+	// True when the parent particle's bone is the direct skeleton parent
+	// of this particle's bone (no intermediate bones skipped). Allows
+	// computing child global position from parent transform + child local.
+	bool bParentParticleIsDirectParent = false;
+
+	// The previous/current update targets
+	FSimVector PreviousTargetPosition = FSimVector::ZeroVector;
+	FSimVector TargetPosition = FSimVector::ZeroVector;
+	FSimVector PreviousTargetVelocity = FSimVector::ZeroVector;
+	FSimVector TargetVelocity = FSimVector::ZeroVector;
+
+	// Normalized direction from parent particle to this particle in the target/animation pose
+	FSimVector TargetDirectionFromParent = FSimVector::ZeroVector;
+	FSimVector PreviousTargetDirectionFromParent = FSimVector::ZeroVector;
+
+	float GravityMultiplier = 1.0f;
+	FSimVector ExternalForce = FSimVector::ZeroVector;
+	ERigParticleSimulationMovementType MovementType = ERigParticleSimulationMovementType::Simulated;
+};
+
+//======================================================================================================================
+struct FParticle
+{
+	// Call once per update
+	void Update(float InInvMass) 
+	{ 
+		InvMass = InInvMass; 
+	}
+
+	float InvMass = 0.0f; // 0 means kinematic (and track the target)
+
+	FSimVector Position = FSimVector::ZeroVector;
+	FSimVector PrevPosition = FSimVector::ZeroVector;
+
+	// The target position is used for kinematic motion
+	FSimVector TargetPosition = FSimVector::ZeroVector;
+	FSimVector TargetVelocity = FSimVector::ZeroVector;
+};
+
+//======================================================================================================================
+struct FParticleTarget
+{
+	// Call once per update. InParentParticle is needed for bilateral compliance when TargetMode > 0.
+	void Update(
+		const FParticle& InParticle, const FParticle* InParentParticle,
+		float InStrength, float InDampingRatio, float InExtraDamping,
+		float InTargetVelocityInfluence, float InTargetMode);
+
+	// How we drive to the target
+	float Compliance = 0.0f; // alpha
+	float Damping = 0.0f;    // beta
+	float TargetVelocityInfluence = 0.0f;
+	// 0 = SimSpace (child tracks absolute target), 1 = directional (child+parent track
+	// animation direction). Values between blend the two modes.
+	float TargetMode = 0.5f;
+
+	// The current interpolated direction from parent to child (unit vector, per substep)
+	FSimVector TargetDirectionFromParent = FSimVector::ZeroVector;
+
+	// Accumulated
+	FSimVector Lambda = FSimVector::ZeroVector;
+};
+
+//======================================================================================================================
+struct FParticleCollider
+{
+	// Call once per update
+	void Update(float InRadius) 
+	{ 
+		Radius = InRadius; 
+	}
+
+	// Collision
+	float Radius = 0.0f;
+	bool bCollideWithColliders = true;
+};
+
+//======================================================================================================================
+struct FParticleNoCollision
+{
+	// Flag per simulation state colliders indicating whether they should collide with this particle
+	TBitArray<> NoCollisionColliderIndices;
+};
+
+//======================================================================================================================
+struct FParticleToParticleCollision
+{
+	// Indices (into the simulation state particles) which this particle will collide with
+	TArray<int32> CollisionParticleIndices;
+};
+
+//======================================================================================================================
+struct FHardDistanceConstraintInfo
+{
+	int32 ParentIndex = -1;
+	int32 ChildIndex = -1;
+};
+
+//======================================================================================================================
+struct FHardDistanceConstraint
+{
+	// Call once per update
+	void Update(float InTargetDistance) { TargetDistance = InTargetDistance; }
+
+	float TargetDistance = 0.0f;
+};
+
+//======================================================================================================================
+struct FSoftDistanceConstraintInfo
+{
+	int32 ParentIndex = -1;
+	int32 ChildIndex = -1;
+};
+
+//======================================================================================================================
+struct FSoftDistanceConstraint
+{
+	// Call once per update
+	void Update(
+		const FParticle& ParentParticle, const FParticle& ChildParticle, 
+		float InTargetDistance, float InStrength, float InDampingRatio, float InExtraDamping);
+
+	float TargetDistance = 0.0f;
+
+	float Compliance = 0.0f; // alpha
+	float Damping = 0.0f;    // beta
+
+	// Accumulated
+	float Lambda = 0.0f;
+};
+
+//======================================================================================================================
+// TODO investigate using the implicit shapes here (and for the others) - e.g. TBox/FImplicitBox3
+struct FBoxCollider
+{
+	// Call once per update
+	void Update(const FSimTransform& InTM, const FSimVector& InExtents)
+	{
+		TM = InTM;
+		Extents = InExtents;
+	}
+
+	FSimTransform TM;
+	FSimVector    Extents = FSimVector::ZeroVector; // Full extents
+};
+
+//======================================================================================================================
+struct FCapsuleCollider
+{
+	// Call once per update
+	void Update(const FSimTransform& InTM, float InLength, float InRadius)
+	{
+		TM = InTM;
+		Length = InLength;
+		Radius = InRadius;
+	}
+
+	// Axis is along +ve Z
+	FSimTransform TM;
+	float         Length = 0.0f; // Full length
+	float         Radius = 0.0f;
+};
+
+//======================================================================================================================
+struct FPlaneCollider
+{
+	// Call once per update
+	void Update(const FSimTransform& InTM, const FSimVector2D& InExtents)
+	{
+		TM = InTM;
+		Extents = InExtents;
+	}
+
+	// Plane points in +ve Z
+	FSimTransform TM;
+	FSimVector2D  Extents = FSimVector2D::ZeroVector;
+};
+
+//======================================================================================================================
+struct FCollider
+{
+	// Call once per update
+	void Update(const FSimTransform& InTM)
+	{
+		PreviousTM = TargetTM;
+		TargetTM = InTM;
+	}
+	FSimTransform PreviousTM;
+	FSimTransform TargetTM;
+
+	// This is the current/interpolated TM
+	FSimTransform TM;
+
+	TArray<FBoxCollider>     BoxColliders;
+	TArray<FCapsuleCollider> CapsuleColliders;
+	TArray<FPlaneCollider>   PlaneColliders;
+};
+
+//======================================================================================================================
+// The structs above here define the individual components that could go into a particle simulation
+// 
+// The structs/functions below here define what components are actually being used for a particular
+// particle simulation (at the moment, one based on skeletal animation/simulation).
+// 
+// In order to turn this into a more general-purpose solver (i.e. one that isn't 100% tailored to
+// our current implementation of control rig dynamics, we will want the simulation state, and the
+// functions that use it, to be more adaptable, and able to specify their own particular structs,
+// constraint types, and order of operations. At that point we can consider using
+// FManagedArrayCollection to store the data, and the Simulate call below can then be supplied with
+// a "recipe" for how the solve should proceed, and what it expects to find.
+//======================================================================================================================
+
+//======================================================================================================================
+// This contains all of the simulation state, customized for a particular simulation setup.
+// Currently, the "user" needs to make and hold this, and then call the functions below in a
+// specific order/way, starting with the Updates, and then calling Simulate.
+//
+// Each block of arrays must have matching sizes - they're separate so that data will be cached better
+//======================================================================================================================
+struct FSimulationState
+{
+	TArray<FParticleInfo>                ParticleInfos;
+	TArray<FParticle>                    Particles;
+	TArray<FParticleTarget>              ParticleTargets;
+	TArray<FParticleCollider>            ParticleColliders;
+	TArray<FParticleNoCollision>         ParticleNoCollision;  
+	TArray<FParticleToParticleCollision> ParticleToParticleCollision;
+
+	TArray<FHardDistanceConstraintInfo>  SkeletalConstraintInfos;     // Along the bones
+	TArray<FHardDistanceConstraint>      SkeletalConstraints;     
+
+	TArray<FHardDistanceConstraintInfo>  HardDistanceConstraintInfos; // User-specified
+	TArray<FHardDistanceConstraint>      HardDistanceConstraints; 
+
+	TArray<FSoftDistanceConstraintInfo>  SoftDistanceConstraintInfos; // User-specified
+	TArray<FSoftDistanceConstraint>      SoftDistanceConstraints; 
+
+	TArray<FCollider>                    Colliders;
+
+	// Previous substep time, used for time-corrected Verlet integration at frame boundaries
+	float PrevStepTime = 0.0f;
+};
+
+//======================================================================================================================
+struct FSolverSettings
+{
+	FSimVector Gravity = FSimVector::ZeroVector;
+	// The total number of iterations, which includes the target solving and the constraint solving
+	int32      NumIterations = 1;
+	// The number of constraint (sub) iterations, processed for every one of the main iterations
+	int32      NumConstraintSubIterations = 1;
+	// The maximum timestep that will be applied, so if the desired step time is greater than this,
+	// substeps will be taken.
+	float      MaxTimeStep = 1.0f;
+	// The maximum number of steps. If this is not enough, then the simulation will run slow.
+	int32      MaxNumSteps = 1;
+};
+
+//======================================================================================================================
+// Describes the motion of the simulation's reference frame, so the simulation can apply
+// inertial pseudo-forces and ether drag. All quantities are in simulation space (the same
+// space as the particles). Populated by the solver before calling Simulate().
+//======================================================================================================================
+struct FSimulationSpaceMotion
+{
+	FSimVector LinearVelocity = FSimVector::ZeroVector;
+	FSimVector AngularVelocity = FSimVector::ZeroVector;      // rad/s
+	FSimVector LinearAcceleration = FSimVector::ZeroVector;
+	FSimVector AngularAcceleration = FSimVector::ZeroVector;   // rad/s^2
+
+	// Wind-like drag force, applied as mass-dependent acceleration (multiplied by InvMass)
+	FSimVector ExternalDragForce = FSimVector::ZeroVector;
+
+	float LinearDragMultiplier = 0.0f;   // Ether coupling for linear frame motion [0,1]
+	float AngularDragMultiplier = 0.0f;  // Ether coupling for angular frame motion [0,1]
+	float SpaceMovementAmount = 0.0f;    // Overall multiplier on all effects [0,1]
+
+	bool bAccelerationsValid = false;
+};
+
+//======================================================================================================================
+// Optionally updates a single particle's simulation properties for the current frame. If the values
+// have changed, it should be called once per particle, before simulation.
+//
+// Index                     - particle index into the simulation state arrays
+// InMass                    - particle mass (used with movement type to compute InvMass)
+// InGravityMultiplier       - scales solver gravity for this particle
+// InRadius                  - collision radius
+//======================================================================================================================
+UE_API void UpdateParticle(
+	FSimulationState&                  State,
+	int32                              Index,
+	float                              InMass,
+	float                              InGravityMultiplier,
+	float                              InRadius);
+
+//======================================================================================================================
+// Updates a single particle's target position and properties for the current frame. Should be
+// called once per particle per update, unless the values and target position haven't changed 
+// for two updates (since velocity is calculated)
+//
+// Index                     - particle index into the simulation state arrays
+// InvDeltaTime              - 1 / DeltaTime (or 0 if DeltaTime is near zero)
+// InTargetPosition          - world-space target position from animation
+// InStrength                - target tracking oscillation frequency
+// InDampingRatio            - target tracking damping ratio
+// InExtraDamping            - additional damping
+// InTargetVelocityInfluence - how much target velocity contributes to damping
+// InTargetMode              - 0 = SimSpace, 1 = directional, blend in between
+//======================================================================================================================
+UE_API void UpdateParticleTarget(
+	FSimulationState&                  State,
+	int32                              Index,
+	float                              InvDeltaTime,
+	const FSimVector&                  InTargetPosition,
+	float                              InStrength,
+	float                              InDampingRatio,
+	float                              InExtraDamping,
+	float                              InTargetVelocityInfluence,
+	float                              InTargetMode);
+
+//======================================================================================================================
+// Updates a single skeletal (bone-chain) distance constraint's target distance from the current
+// particle target positions. Call once per constraint per update when bone lengths may change (but skip if they can't).
+//======================================================================================================================
+UE_API void UpdateSkeletalConstraint(FSimulationState& State, int32 Index);
+
+//======================================================================================================================
+// Updates a single user-defined hard distance constraint for the current frame. Recomputes the
+// target distance from particle target positions, scaled by InLengthMultiplier and offset by
+// InExtraLength. Call once per constraint per update when values have changed.
+//======================================================================================================================
+UE_API void UpdateHardDistanceConstraint(
+	FSimulationState& State, int32 Index, float InLengthMultiplier, float InExtraLength);
+
+//======================================================================================================================
+// Updates a single user-defined soft distance constraint for the current frame. Recomputes the
+// target distance and compliance/damping from the given parameters. Call once per constraint per 
+// update when values have changed.
+//======================================================================================================================
+UE_API void UpdateSoftDistanceConstraint(
+	FSimulationState& State, int32 Index, float InLengthMultiplier, float InExtraLength,
+	float InStrength, float InDampingRatio, float InExtraDamping);
+
+//======================================================================================================================
+// Updates a single collider for the current frame. Call once per collider per update when the
+// properties have changed.
+//======================================================================================================================
+UE_API void UpdateCollider(FSimulationState& State, int32 Index, const FSimTransform& TM);
+
+//======================================================================================================================
+// Runs the simulation for one frame. Internally substeps according to MaxTimeStep/MaxNumSteps.
+// The caller should have called the Update functions above to set particle targets, constraint
+// distances, and collider transforms for the current frame before calling this.
+// Particles should be in root-to-leaf order (see SortParticlesRootToLeaf) so that parent
+// particles are processed before their children during constraint solving and readback.
+// SimulationSpaceMotion is used to pass information about movement of the simulation space itself
+// so that effects of non-inertial motion etc can be applied to the simulated particles.
+//======================================================================================================================
+UE_API void Simulate(
+	FSimulationState&              SimulationState,
+	const FSolverSettings&         SolverSettings,
+	const FSimulationSpaceMotion&  SimulationSpaceMotion,
+	float                          DeltaTime);
+
+//======================================================================================================================
+// Sorts all particle-parallel arrays in the simulation state into root-to-leaf (BFS) order.
+// After sorting, ParticleInfos[i].ParentParticleIndex < i for all non-root particles.
+// All stored particle index references are remapped (ParentParticleIndex, constraint
+// ParentIndex/ChildIndex, and particle-to-particle collision indices), so this can be called
+// at any point after particles and their parent indices have been set up.
+// Returns the permutation: Result[new_index] = old_index, so the caller can reorder any
+// additional parallel arrays it owns outside FSimulationState.
+//======================================================================================================================
+UE_API TArray<int32> SortParticlesRootToLeaf(FSimulationState& State);
+
+//======================================================================================================================
+// Trivial/inline implementations
+//======================================================================================================================
+
+//======================================================================================================================
+inline void UpdateSkeletalConstraint(FSimulationState& State, const int32 Index)
+{
+	const FHardDistanceConstraintInfo& Info = State.SkeletalConstraintInfos[Index];
+	check(State.ParticleInfos.IsValidIndex(Info.ParentIndex) && State.ParticleInfos.IsValidIndex(Info.ChildIndex));
+	const FSimVector& ParentPos = State.ParticleInfos[Info.ParentIndex].TargetPosition;
+	const FSimVector& ChildPos = State.ParticleInfos[Info.ChildIndex].TargetPosition;
+	State.SkeletalConstraints[Index].Update(FSimVector::Dist(ParentPos, ChildPos));
+}
+
+//======================================================================================================================
+inline void UpdateHardDistanceConstraint(
+	FSimulationState& State, const int32 Index, const float InLengthMultiplier, const float InExtraLength)
+{
+	const FHardDistanceConstraintInfo& Info = State.HardDistanceConstraintInfos[Index];
+	check(State.ParticleInfos.IsValidIndex(Info.ParentIndex) && State.ParticleInfos.IsValidIndex(Info.ChildIndex));
+	const FSimVector& ParentPos = State.ParticleInfos[Info.ParentIndex].TargetPosition;
+	const FSimVector& ChildPos = State.ParticleInfos[Info.ChildIndex].TargetPosition;
+	const float Distance = FSimVector::Dist(ParentPos, ChildPos) * InLengthMultiplier + InExtraLength;
+	State.HardDistanceConstraints[Index].Update(Distance);
+}
+
+//======================================================================================================================
+inline void UpdateSoftDistanceConstraint(
+	FSimulationState& State, const int32 Index, const float InLengthMultiplier, const float InExtraLength,
+	const float InStrength, const float InDampingRatio, const float InExtraDamping)
+{
+	const FSoftDistanceConstraintInfo& Info = State.SoftDistanceConstraintInfos[Index];
+	check(State.ParticleInfos.IsValidIndex(Info.ParentIndex) && State.ParticleInfos.IsValidIndex(Info.ChildIndex));
+	const FSimVector& ParentPos = State.ParticleInfos[Info.ParentIndex].TargetPosition;
+	const FSimVector& ChildPos = State.ParticleInfos[Info.ChildIndex].TargetPosition;
+	const float Distance = FSimVector::Dist(ParentPos, ChildPos) * InLengthMultiplier + InExtraLength;
+	State.SoftDistanceConstraints[Index].Update(State.Particles[Info.ParentIndex], State.Particles[Info.ChildIndex],
+		Distance, InStrength, InDampingRatio, InExtraDamping);
+}
+
+//======================================================================================================================
+inline void UpdateCollider(FSimulationState& State, int32 Index, const FSimTransform& TM)
+{
+	check(State.Colliders.IsValidIndex(Index));
+	RigParticleSimulation::FCollider& Collider = State.Colliders[Index];
+	Collider.Update(TM);
+}
+
+}
+
+#undef UE_API
