Control Rig Dynamics - support for cone limits and for (simple) angular limits on the target directions. Some minor tidyup too.

#Synced-CL 52725359

diff --git a/Engine/Plugins/Experimental/ControlRigDynamics/ControlRigDynamics.uplugin b/Engine/Plugins/Experimental/ControlRigDynamics/ControlRigDynamics.uplugin
--- a/Engine/Plugins/Experimental/ControlRigDynamics/ControlRigDynamics.uplugin
+++ b/Engine/Plugins/Experimental/ControlRigDynamics/ControlRigDynamics.uplugin
@@ -30,6 +30,6 @@
 			"Name": "ControlRigDynamics",
 			"Type": "Runtime",
 			"LoadingPhase": "Default"
-		},
+		}
 	]
 }
\ No newline at end of file
diff --git a/Engine/Plugins/Experimental/ControlRigDynamics/Source/ControlRigDynamics/Private/RigDynamicsConeLimitComponent.cpp b/Engine/Plugins/Experimental/ControlRigDynamics/Source/ControlRigDynamics/Private/RigDynamicsConeLimitComponent.cpp
new file mode 100644
--- /dev/null
+++ b/Engine/Plugins/Experimental/ControlRigDynamics/Source/ControlRigDynamics/Private/RigDynamicsConeLimitComponent.cpp
@@ -0,0 +1,31 @@
+// Copyright Epic Games, Inc. All Rights Reserved.
+
+#include "RigDynamicsConeLimitComponent.h"
+
+#include "RigDynamicsObjectVersion.h"
+
+//======================================================================================================================
+void FRigDynamicsConeLimitComponent::Serialize(FArchive& Ar)
+{
+	Ar << GrandparentComponentKey;
+	Ar << ParentComponentKey;
+	Ar << ChildComponentKey;
+	Ar << Strength;
+	Ar << DampingRatio;
+	Ar << Angle;
+}
+
+//======================================================================================================================
+void FRigDynamicsConeLimitComponent::Save(FArchive& Ar)
+{
+	Ar.UsingCustomVersion(FRigDynamicsObjectVersion::GUID);
+	FRigBaseComponent::Save(Ar);
+	Serialize(Ar);
+}
+
+//======================================================================================================================
+void FRigDynamicsConeLimitComponent::Load(FArchive& Ar)
+{
+	FRigBaseComponent::Load(Ar);
+	Serialize(Ar);
+}
diff --git a/Engine/Plugins/Experimental/ControlRigDynamics/Source/ControlRigDynamics/Private/RigDynamicsConeLimitExecution.cpp b/Engine/Plugins/Experimental/ControlRigDynamics/Source/ControlRigDynamics/Private/RigDynamicsConeLimitExecution.cpp
new file mode 100644
--- /dev/null
+++ b/Engine/Plugins/Experimental/ControlRigDynamics/Source/ControlRigDynamics/Private/RigDynamicsConeLimitExecution.cpp
@@ -0,0 +1,89 @@
+// Copyright Epic Games, Inc. All Rights Reserved.
+
+#include "RigDynamicsConeLimitExecution.h"
+
+#include "RigDynamicsHelpers.h"
+
+#include "Rigs/RigHierarchyController.h"
+#include "Units/RigUnitContext.h"
+#include "Units/Execution/RigUnit_PrepareForExecution.h"
+
+#include UE_INLINE_GENERATED_CPP_BY_NAME(RigDynamicsConeLimitExecution)
+
+//======================================================================================================================
+// Macros to reduce boilerplate for the Get/Set cone limit property nodes. Note that these can't
+// be used in the header because UHT scans raw source before C++ preprocessing, so it would
+// never see the USTRUCT/UPROPERTY/RIGVM_METHOD markers inside a macro expansion.
+//======================================================================================================================
+
+#define IMPLEMENT_SET_DYNAMICS_CONE_LIMIT(PropName)                               \
+FRigUnit_HierarchySetDynamicsConeLimit##PropName##_Execute()                      \
+{                                                                                 \
+	if (!ExecuteContext.Hierarchy) { return; }                                    \
+	if (FRigDynamicsConeLimitComponent* Component =                               \
+		GetConeLimit(*ExecuteContext.Hierarchy, DynamicsConeLimitComponentKey))   \
+	{                                                                             \
+		Component->PropName = PropName;                                           \
+	}                                                                             \
+}
+
+#define IMPLEMENT_GET_DYNAMICS_CONE_LIMIT(PropName)                               \
+FRigUnit_HierarchyGetDynamicsConeLimit##PropName##_Execute()                      \
+{                                                                                 \
+	if (!ExecuteContext.Hierarchy) { return; }                                    \
+	if (const FRigDynamicsConeLimitComponent* Component =                         \
+		GetConeLimit(*ExecuteContext.Hierarchy, DynamicsConeLimitComponentKey))   \
+	{                                                                             \
+		PropName = Component->PropName;                                           \
+	}                                                                             \
+}
+
+//======================================================================================================================
+FRigUnit_SpawnDynamicsConeLimit_Execute()
+{
+	if (!ExecuteContext.Hierarchy)
+	{
+		return;
+	}
+	if (ExecuteContext.GetEventName() != FRigUnit_PrepareForExecution::EventName)
+	{
+		UE_CONTROLRIG_RIGUNIT_REPORT_ERROR(TEXT("SpawnDynamicsConeLimit can only be used during Setup"));
+		return;
+	}
+
+	if (URigHierarchyController* Controller = ExecuteContext.Hierarchy->GetController())
+	{
+		FRigHierarchyControllerInstructionBracket InstructionBracket(
+			Controller, ExecuteContext.GetInstructionIndex());
+		DynamicsConeLimitComponentKey = Controller->AddComponent(
+			FRigDynamicsConeLimitComponent::StaticStruct(), FRigDynamicsConeLimitComponent::GetDefaultName(), Owner);
+		if (DynamicsConeLimitComponentKey.IsValid())
+		{
+			if (FRigDynamicsConeLimitComponent* Component =
+				GetConeLimit(*ExecuteContext.Hierarchy, DynamicsConeLimitComponentKey))
+			{
+				Component->GrandparentComponentKey = GrandparentComponentKey;
+				Component->ParentComponentKey = ParentComponentKey;
+				Component->ChildComponentKey = ChildComponentKey;
+				Component->Strength = Strength;
+				Component->DampingRatio = DampingRatio;
+				Component->Angle = Angle;
+			}
+
+			if (FRigDynamicsSolverComponent* Solver = GetSolver(*ExecuteContext.Hierarchy, DynamicsSolverComponentKey))
+			{
+				Solver->ConeLimits.Add(DynamicsConeLimitComponentKey);
+			}
+		}
+	}
+}
+
+//======================================================================================================================
+IMPLEMENT_SET_DYNAMICS_CONE_LIMIT(Strength)
+IMPLEMENT_GET_DYNAMICS_CONE_LIMIT(Strength)
+//======================================================================================================================
+IMPLEMENT_SET_DYNAMICS_CONE_LIMIT(DampingRatio)
+IMPLEMENT_GET_DYNAMICS_CONE_LIMIT(DampingRatio)
+//======================================================================================================================
+IMPLEMENT_SET_DYNAMICS_CONE_LIMIT(Angle)
+IMPLEMENT_GET_DYNAMICS_CONE_LIMIT(Angle)
diff --git a/Engine/Plugins/Experimental/ControlRigDynamics/Source/ControlRigDynamics/Private/RigDynamicsHelpers.h b/Engine/Plugins/Experimental/ControlRigDynamics/Source/ControlRigDynamics/Private/RigDynamicsHelpers.h
--- a/Engine/Plugins/Experimental/ControlRigDynamics/Source/ControlRigDynamics/Private/RigDynamicsHelpers.h
+++ b/Engine/Plugins/Experimental/ControlRigDynamics/Source/ControlRigDynamics/Private/RigDynamicsHelpers.h
@@ -6,6 +6,7 @@
 #include "RigDynamicsSolverComponent.h"
 #include "RigDynamicsColliderComponent.h"
 #include "RigDynamicsConstraintComponent.h"
+#include "RigDynamicsConeLimitComponent.h"
 
 #include "Rigs/RigHierarchy.h"
 
@@ -89,6 +90,26 @@ inline const FRigDynamicsConstraintComponent* GetConstraint(
 	return Cast<FRigDynamicsConstraintComponent>(CachedRigComponent.GetComponent(&Hierarchy));
 }
 
+//======================================================================================================================
+inline FRigDynamicsConeLimitComponent* GetConeLimit(URigHierarchy& Hierarchy, const FRigComponentKey& Key)
+{
+	return Cast<FRigDynamicsConeLimitComponent>(Hierarchy.FindComponent(Key));
+}
+
+//======================================================================================================================
+inline const FRigDynamicsConeLimitComponent* GetConeLimit(const URigHierarchy& Hierarchy, const FRigComponentKey& Key)
+{
+	return Cast<FRigDynamicsConeLimitComponent>(Hierarchy.FindComponent(Key));
+}
+
+//======================================================================================================================
+// Updates the cache if necessary and then returns the component
+inline const FRigDynamicsConeLimitComponent* GetConeLimit(
+	const URigHierarchy& Hierarchy, FCachedRigComponent& CachedRigComponent)
+{
+	return Cast<FRigDynamicsConeLimitComponent>(CachedRigComponent.GetComponent(&Hierarchy));
+}
+
 //======================================================================================================================
 // Gets the transform from the element that owns this component. 
 inline FTransform GetGlobalTransform(const URigHierarchy& Hierarchy, const FRigBaseComponent& Component)
diff --git a/Engine/Plugins/Experimental/ControlRigDynamics/Source/ControlRigDynamics/Private/RigDynamicsObjectVersion.h b/Engine/Plugins/Experimental/ControlRigDynamics/Source/ControlRigDynamics/Private/RigDynamicsObjectVersion.h
--- a/Engine/Plugins/Experimental/ControlRigDynamics/Source/ControlRigDynamics/Private/RigDynamicsObjectVersion.h
+++ b/Engine/Plugins/Experimental/ControlRigDynamics/Source/ControlRigDynamics/Private/RigDynamicsObjectVersion.h
@@ -24,6 +24,8 @@ struct FRigDynamicsObjectVersion
 		CollisionParticles,
 		TargetModeFloat,
 		SimulationSpace,
+		ConeLimits,
+		AngleLimit,
 
 		// -----<new versions can be added above this line>-----
 		VersionPlusOne,
diff --git a/Engine/Plugins/Experimental/ControlRigDynamics/Source/ControlRigDynamics/Private/RigDynamicsParticleComponent.cpp b/Engine/Plugins/Experimental/ControlRigDynamics/Source/ControlRigDynamics/Private/RigDynamicsParticleComponent.cpp
--- a/Engine/Plugins/Experimental/ControlRigDynamics/Source/ControlRigDynamics/Private/RigDynamicsParticleComponent.cpp
+++ b/Engine/Plugins/Experimental/ControlRigDynamics/Source/ControlRigDynamics/Private/RigDynamicsParticleComponent.cpp
@@ -21,11 +21,15 @@ FArchive& operator <<(FArchive& Ar, FRigDynamicsParticleProperties& Data)
 	}
 	else
 	{
-		// Old data: TargetMode was ERigParticleSimulationTargetMode (uint8).
-		// SimSpace=0 -> 0.0f, PreservePose=1 -> 1.0f
+		// Ignore old data (we have no legacy assets)
 		uint8 LegacyTargetMode = 1;
 		Ar << LegacyTargetMode;
 	}
+	if (Ar.CustomVer(FRigDynamicsObjectVersion::GUID) >= FRigDynamicsObjectVersion::AngleLimit)
+	{
+		Ar << Data.AngleLimit;
+		Ar << Data.AngleLimitStrength;
+	}
 	if (Ar.CustomVer(FRigDynamicsObjectVersion::GUID) >= FRigDynamicsObjectVersion::NoCollisionColliders)
 	{
 		Ar << Data.NoCollisionColliders;
@@ -63,7 +67,7 @@ void FRigDynamicsParticleComponent::Serialize(FArchive& Ar)
 		Ar << ParticleProperties.TargetVelocityInfluence;
 		if (Ar.CustomVer(FRigDynamicsObjectVersion::GUID) >= FRigDynamicsObjectVersion::DynamicsTargetMode)
 		{
-			// Pre-HelperStructs is always pre-TargetModeFloat, so always read the legacy uint8
+			// Ignore old data (we have no legacy assets)
 			uint8 LegacyTargetMode = 1;
 			Ar << LegacyTargetMode;
 		}
diff --git a/Engine/Plugins/Experimental/ControlRigDynamics/Source/ControlRigDynamics/Private/RigDynamicsParticleExecution.cpp b/Engine/Plugins/Experimental/ControlRigDynamics/Source/ControlRigDynamics/Private/RigDynamicsParticleExecution.cpp
--- a/Engine/Plugins/Experimental/ControlRigDynamics/Source/ControlRigDynamics/Private/RigDynamicsParticleExecution.cpp
+++ b/Engine/Plugins/Experimental/ControlRigDynamics/Source/ControlRigDynamics/Private/RigDynamicsParticleExecution.cpp
@@ -148,3 +148,9 @@ IMPLEMENT_GET_DYNAMICS_PARTICLE(TargetVelocityInfluence)
 //======================================================================================================================
 IMPLEMENT_SET_DYNAMICS_PARTICLE(TargetMode)
 IMPLEMENT_GET_DYNAMICS_PARTICLE(TargetMode)
+//======================================================================================================================
+IMPLEMENT_SET_DYNAMICS_PARTICLE(AngleLimit)
+IMPLEMENT_GET_DYNAMICS_PARTICLE(AngleLimit)
+//======================================================================================================================
+IMPLEMENT_SET_DYNAMICS_PARTICLE(AngleLimitStrength)
+IMPLEMENT_GET_DYNAMICS_PARTICLE(AngleLimitStrength)
diff --git a/Engine/Plugins/Experimental/ControlRigDynamics/Source/ControlRigDynamics/Private/RigDynamicsSimulationSpace.cpp b/Engine/Plugins/Experimental/ControlRigDynamics/Source/ControlRigDynamics/Private/RigDynamicsSimulationSpace.cpp
--- a/Engine/Plugins/Experimental/ControlRigDynamics/Source/ControlRigDynamics/Private/RigDynamicsSimulationSpace.cpp
+++ b/Engine/Plugins/Experimental/ControlRigDynamics/Source/ControlRigDynamics/Private/RigDynamicsSimulationSpace.cpp
@@ -105,6 +105,7 @@ void FRigDynamicsSimulationSpaceState::Update(
 	// If delta time is zero just update the transform and don't worry too much about velocities being off.
 	SimulationSpaceTM = InSimulationSpaceTM;
 	ComponentTM = InComponentTM;
+	RecomputeCompositeTransforms();
 }
 
 //======================================================================================================================
@@ -114,7 +115,19 @@ void FRigDynamicsSimulationSpaceState::Reset(
 	*this = FRigDynamicsSimulationSpaceState();
 	SimulationSpaceTM = InSimulationSpaceTM;
 	ComponentTM = InComponentTM;
-	UpdatesSinceReset = 0;
+	RecomputeCompositeTransforms();
+}
+
+//======================================================================================================================
+void FRigDynamicsSimulationSpaceState::RecomputeCompositeTransforms()
+{
+	// ComponentToSimSpace: component-space pos -> world (via ComponentTM) -> sim-space (via SimTM^-1)
+	// In UE composition (A*B means "apply A then B"): ComponentTM * SimulationSpaceTM.Inverse()
+	ComponentToSimSpaceTM = ComponentTM * SimulationSpaceTM.Inverse();
+	ComponentToSimSpaceTM.SetScale3D(FVector::OneVector);
+
+	SimToComponentSpaceTM = SimulationSpaceTM * ComponentTM.Inverse();
+	SimToComponentSpaceTM.SetScale3D(FVector::OneVector);
 }
 
 //======================================================================================================================
diff --git a/Engine/Plugins/Experimental/ControlRigDynamics/Source/ControlRigDynamics/Private/RigDynamicsSolver.cpp b/Engine/Plugins/Experimental/ControlRigDynamics/Source/ControlRigDynamics/Private/RigDynamicsSolver.cpp
--- a/Engine/Plugins/Experimental/ControlRigDynamics/Source/ControlRigDynamics/Private/RigDynamicsSolver.cpp
+++ b/Engine/Plugins/Experimental/ControlRigDynamics/Source/ControlRigDynamics/Private/RigDynamicsSolver.cpp
@@ -93,6 +93,7 @@ void FRigDynamicsSolver::Instantiate(
 	HardDistanceConstraintOwnerComponents.Reset();
 	SoftDistanceConstraintOwnerComponents.Reset();
 	ColliderOwnerComponents.Reset();
+	ConeLimitOwnerComponents.Reset();
 
 	SimulationState = RigParticleSimulation::FSimulationState();
 
@@ -192,7 +193,6 @@ void FRigDynamicsSolver::Instantiate(
 	const int32 NumParticles = SimulationState.Particles.Num();
 
 	// Store the parent particle indices - these will be needed for targets and joint limits
-	// TODO implement angular (cone) limits, which will need the ParentParentParticleIndex
 	for (int32 ParticleIndex = 0; ParticleIndex != NumParticles; ++ParticleIndex)
 	{
 		RigParticleSimulation::FParticleInfo& ParticleInfo = SimulationState.ParticleInfos[ParticleIndex];
@@ -371,6 +371,64 @@ void FRigDynamicsSolver::Instantiate(
 		}
 	}
 
+	// User-defined cone limits
+	for (const FRigComponentKey& ConeLimitComponentKey : SolverComponent.ConeLimits)
+	{
+		if (const FRigDynamicsConeLimitComponent* CC = GetConeLimit(Hierarchy, ConeLimitComponentKey))
+		{
+			int32 GrandparentIndex = GetParticleIndexSlow(CC->GrandparentComponentKey);
+			int32 ParentIndex = GetParticleIndexSlow(CC->ParentComponentKey);
+			int32 ChildIndex = GetParticleIndexSlow(CC->ChildComponentKey);
+
+			if (GrandparentIndex == INDEX_NONE || ParentIndex == INDEX_NONE || ChildIndex == INDEX_NONE)
+			{
+				continue;
+			}
+
+			RigParticleSimulation::FConeLimitInfo& Info =
+				SimulationState.ConeLimitInfos.AddDefaulted_GetRef();
+			SimulationState.ConeLimits.AddDefaulted();
+			ConeLimitOwnerComponents.Add(FCachedRigComponent(ConeLimitComponentKey, &Hierarchy, true));
+
+			Info.GrandparentIndex = GrandparentIndex;
+			Info.ParentIndex = ParentIndex;
+			Info.ChildIndex = ChildIndex;
+		}
+	}
+
+	// Sort cone limits by child index (root-to-leaf) for better iterative convergence.
+	if (SimulationState.ConeLimits.Num() > 1)
+	{
+		const int32 NumConeLimits = SimulationState.ConeLimits.Num();
+		TArray<int32> ConeLimitOrder;
+		ConeLimitOrder.SetNum(NumConeLimits);
+		for (int32 I = 0; I < NumConeLimits; ++I)
+		{
+			ConeLimitOrder[I] = I;
+		}
+		ConeLimitOrder.Sort([&](int32 A, int32 B)
+		{
+			return SimulationState.ConeLimitInfos[A].ChildIndex < SimulationState.ConeLimitInfos[B].ChildIndex;
+		});
+
+		TArray<RigParticleSimulation::FConeLimitInfo> TempInfos;
+		TArray<RigParticleSimulation::FConeLimit> TempLimits;
+		TArray<FCachedRigComponent> TempOwners;
+		TempInfos.SetNum(NumConeLimits);
+		TempLimits.SetNum(NumConeLimits);
+		TempOwners.SetNum(NumConeLimits);
+
+		for (int32 I = 0; I < NumConeLimits; ++I)
+		{
+			TempInfos[I]  = MoveTemp(SimulationState.ConeLimitInfos[ConeLimitOrder[I]]);
+			TempLimits[I] = MoveTemp(SimulationState.ConeLimits[ConeLimitOrder[I]]);
+			TempOwners[I] = MoveTemp(ConeLimitOwnerComponents[ConeLimitOrder[I]]);
+		}
+		SimulationState.ConeLimitInfos = MoveTemp(TempInfos);
+		SimulationState.ConeLimits = MoveTemp(TempLimits);
+		ConeLimitOwnerComponents = MoveTemp(TempOwners);
+	}
+
 	bNeedToInstantiate = false;
 }
 
@@ -400,9 +458,10 @@ void FRigDynamicsSolver::UpdatePreDynamics(URigHierarchy& Hierarchy, float Delta
 					SimulationSpaceState.ConvertComponentSpacePositionToSimSpace(TargetTM.GetLocation());
 				RigParticleSimulation::UpdateParticleTarget(
 					SimulationState, Index, InvDeltaTime, FSimVector(TargetSimSpace),
-					PC->ParticleProperties.Strength, PC->ParticleProperties.DampingRatio, 
-					PC->ParticleProperties.ExtraDamping, PC->ParticleProperties.TargetVelocityInfluence, 
-					PC->ParticleProperties.TargetMode);
+					PC->ParticleProperties.Strength, PC->ParticleProperties.DampingRatio,
+					PC->ParticleProperties.ExtraDamping, PC->ParticleProperties.TargetVelocityInfluence,
+					PC->ParticleProperties.TargetMode, PC->ParticleProperties.AngleLimit,
+					PC->ParticleProperties.AngleLimitStrength);
 			}
 		}
 	}
@@ -442,6 +501,19 @@ void FRigDynamicsSolver::UpdatePreDynamics(URigHierarchy& Hierarchy, float Delta
 					CC->LengthMultiplier, CC->ExtraLength, CC->Strength, CC->DampingRatio, CC->ExtraDamping);
 			}
 		}
+
+		// Update user-specified cone limits
+		check(SimulationState.ConeLimits.Num() == SimulationState.ConeLimitInfos.Num());
+		check(ConeLimitOwnerComponents.Num() == SimulationState.ConeLimitInfos.Num());
+		for (int32 Index = 0; Index != SimulationState.ConeLimits.Num(); ++Index)
+		{
+			if (const FRigDynamicsConeLimitComponent* CC = GetConeLimit(
+				Hierarchy, ConeLimitOwnerComponents[Index]))
+			{
+				RigParticleSimulation::UpdateConeLimit(
+					SimulationState, Index, CC->Angle, CC->Strength, CC->DampingRatio);
+			}
+		}
 	}
 
 	{
diff --git a/Engine/Plugins/Experimental/ControlRigDynamics/Source/ControlRigDynamics/Private/RigDynamicsSolver.h b/Engine/Plugins/Experimental/ControlRigDynamics/Source/ControlRigDynamics/Private/RigDynamicsSolver.h
--- a/Engine/Plugins/Experimental/ControlRigDynamics/Source/ControlRigDynamics/Private/RigDynamicsSolver.h
+++ b/Engine/Plugins/Experimental/ControlRigDynamics/Source/ControlRigDynamics/Private/RigDynamicsSolver.h
@@ -96,6 +96,7 @@ private:
 	TArray<FCachedRigComponent>  HardDistanceConstraintOwnerComponents;
 	TArray<FCachedRigComponent>  SoftDistanceConstraintOwnerComponents;
 	TArray<FCachedRigComponent>  ColliderOwnerComponents;
+	TArray<FCachedRigComponent>  ConeLimitOwnerComponents;
 
 	RigParticleSimulation::FSimulationState SimulationState;
 
diff --git a/Engine/Plugins/Experimental/ControlRigDynamics/Source/ControlRigDynamics/Private/RigDynamicsSolverComponent.cpp b/Engine/Plugins/Experimental/ControlRigDynamics/Source/ControlRigDynamics/Private/RigDynamicsSolverComponent.cpp
--- a/Engine/Plugins/Experimental/ControlRigDynamics/Source/ControlRigDynamics/Private/RigDynamicsSolverComponent.cpp
+++ b/Engine/Plugins/Experimental/ControlRigDynamics/Source/ControlRigDynamics/Private/RigDynamicsSolverComponent.cpp
@@ -14,6 +14,10 @@ void FRigDynamicsSolverComponent::Serialize(FArchive& Ar)
 	{
 		Ar << Constraints;
 	}
+	if (Ar.CustomVer(FRigDynamicsObjectVersion::GUID) >= FRigDynamicsObjectVersion::ConeLimits)
+	{
+		Ar << ConeLimits;
+	}
 	Ar << Settings;
 	if (Ar.CustomVer(FRigDynamicsObjectVersion::GUID) >= FRigDynamicsObjectVersion::SimulationSpace)
 	{
diff --git a/Engine/Plugins/Experimental/ControlRigDynamics/Source/ControlRigDynamics/Private/RigDynamicsSolverExecution.cpp b/Engine/Plugins/Experimental/ControlRigDynamics/Source/ControlRigDynamics/Private/RigDynamicsSolverExecution.cpp
--- a/Engine/Plugins/Experimental/ControlRigDynamics/Source/ControlRigDynamics/Private/RigDynamicsSolverExecution.cpp
+++ b/Engine/Plugins/Experimental/ControlRigDynamics/Source/ControlRigDynamics/Private/RigDynamicsSolverExecution.cpp
@@ -12,7 +12,7 @@
 
 #include UE_INLINE_GENERATED_CPP_BY_NAME(RigDynamicsSolverExecution)
 
-TAutoConsoleVariable<int> CVarControlRigDynamicsVisualizationOverride(
+TAutoConsoleVariable<int32> CVarControlRigDynamicsVisualizationOverride(
 	TEXT("ControlRig.Dynamics.VisualizationOverride"), -1,
 	TEXT("-1 uses the visualization settings from Step Dynamics Solver, 0 forces visualization to be disabled, 1 forces it to be enabled."));
 
diff --git a/Engine/Plugins/Experimental/ControlRigDynamics/Source/ControlRigDynamics/Private/RigParticleSimulation.cpp b/Engine/Plugins/Experimental/ControlRigDynamics/Source/ControlRigDynamics/Private/RigParticleSimulation.cpp
--- a/Engine/Plugins/Experimental/ControlRigDynamics/Source/ControlRigDynamics/Private/RigParticleSimulation.cpp
+++ b/Engine/Plugins/Experimental/ControlRigDynamics/Source/ControlRigDynamics/Private/RigParticleSimulation.cpp
@@ -9,51 +9,6 @@
 namespace RigParticleSimulation
 {
 
-//======================================================================================================================
-void FParticleTarget::Update(
-	const FParticle&  InParticle,
-	const FParticle*  InParentParticle,
-	const float       InStrength,
-	const float       InDampingRatio,
-	const float       InExtraDamping,
-	const float       InTargetVelocityInfluence,
-	const float       InTargetMode)
-{
-	TargetVelocityInfluence = InTargetVelocityInfluence;
-	TargetMode = InTargetMode;
-
-	// Scale compliance by the effective inverse mass so the constraint frequency is
-	// mass-independent. Parent contribution is weighted by Alpha^2 (from the squared Jacobian).
-	const float Alpha = InParentParticle ? InTargetMode : 0.0f;
-	const float InvMass = InParticle.InvMass
-		+ Alpha * Alpha * (InParentParticle ? InParentParticle->InvMass : 0.0f);
-
-	// Angular frequency
-	const float W = InStrength * TWO_PI;
-	// Stiffness = W^2, and compliance is inverse stiffness (scaled by inverse mass)
-	Compliance = InvMass / FMath::Max(W * W, KINDA_SMALL_NUMBER);
-	Damping = (2.0f * InDampingRatio * W + InExtraDamping) / FMath::Max(InvMass, KINDA_SMALL_NUMBER);
-}
-
-//======================================================================================================================
-void FSoftDistanceConstraint::Update(
-	const FParticle& ParentParticle, 
-	const FParticle& ChildParticle,
-	const float      InTargetDistance,
-	const float      InStrength,
-	const float      InDampingRatio,
-	const float      InExtraDamping)
-{
-	TargetDistance = InTargetDistance;
-
-	// Scale compliance by effective inverse mass so the constraint frequency is
-	// mass-independent
-	const float W = InStrength * TWO_PI;
-	const float InvMassSum = ParentParticle.InvMass + ChildParticle.InvMass;
-	Compliance = InvMassSum / FMath::Max(W * W, KINDA_SMALL_NUMBER);
-	Damping = (2.0f * InDampingRatio * W + InExtraDamping) / FMath::Max(InvMassSum, KINDA_SMALL_NUMBER);
-}
-
 //======================================================================================================================
 static void StepParticleDynamic(
 	FParticle&                      Particle,
@@ -181,6 +136,12 @@ static void StepSoftDistanceConstraint(FSoftDistanceConstraint& Constraint)
 	Constraint.Lambda = 0.0f;
 }
 
+//======================================================================================================================
+static void StepConeLimit(FConeLimit& ConeLimit)
+{
+	ConeLimit.Lambda = 0.0f;
+}
+
 //======================================================================================================================
 static void StepCollider(FCollider& Collider, const float StepAlpha)
 {
@@ -265,6 +226,91 @@ static void SolveTargetConstraint(
 	}
 }
 
+//======================================================================================================================
+// Constrains the direction from parent to child to lie within AngleLimit of the target
+// (animation) direction. The correction is scaled by a compliance-derived factor so that low
+// AngleLimitStrength gives soft/springy behaviour. Uses inverse-mass weighting between child
+// and parent. When bOnlyMoveChild is true, only the child is moved.
+static void SolveTargetAngleConstraint(
+	FParticle&              Particle,
+	const FParticleTarget&  ParticleTarget,
+	FParticle*              ParentParticlePtr,
+	const bool              bOnlyMoveChild,
+	const float             InvDeltaTime)
+{
+	if (ParticleTarget.AngleLimitCompliance <= 0.0f || !ParentParticlePtr)
+	{
+		return;
+	}
+
+	if (ParticleTarget.TargetDirectionFromParent.IsNearlyZero(KINDA_SMALL_NUMBER))
+	{
+		return;
+	}
+
+	const FSimVector Delta = Particle.Position - ParentParticlePtr->Position;
+	const float DeltaLen = Delta.Size();
+	if (DeltaLen < KINDA_SMALL_NUMBER)
+	{
+		return;
+	}
+
+	const FSimVector CurrentDir = Delta * (1.0f / DeltaLen);
+	const FSimVector& TargetDir = ParticleTarget.TargetDirectionFromParent;
+
+	const float CosTheta = FMath::Clamp(
+		(float)FSimVector::DotProduct(CurrentDir, TargetDir), -1.0f + KINDA_SMALL_NUMBER, 1.0f);
+
+	// Equivalent to Theta <= AngleLimit
+	if (CosTheta >= ParticleTarget.CosAngleLimit)
+	{
+		return;
+	}
+
+	// Rotation axis to bring CurrentDir toward TargetDir
+	FSimVector Axis = FSimVector::CrossProduct(CurrentDir, TargetDir);
+	const float AxisLen = Axis.Size();
+	if (AxisLen < KINDA_SMALL_NUMBER)
+	{
+		return;
+	}
+	Axis *= (1.0f / AxisLen);
+
+	// Rotate Delta toward target direction by the excess angle (Theta - Limit) using
+	// Rodrigues' formula. dot(Axis, Delta) is zero so it simplifies to:
+	//   NewDelta = Delta * cos(Excess) + cross(Axis, Delta) * sin(Excess)
+	//
+	// We avoid trig on the excess angle via the subtraction identities:
+	//   cos(Theta - Limit) = cosTheta * cosLimit + sinTheta * sinLimit
+	//   sin(Theta - Limit) = sinTheta * cosLimit - cosTheta * sinLimit
+	const float SinTheta = FMath::Sqrt(1.0f - CosTheta * CosTheta);
+	const float CosExcess = CosTheta * ParticleTarget.CosAngleLimit + SinTheta * ParticleTarget.SinAngleLimit;
+	const float SinExcess = SinTheta * ParticleTarget.CosAngleLimit - CosTheta * ParticleTarget.SinAngleLimit;
+	const FSimVector NewDelta = Delta * CosExcess + FSimVector::CrossProduct(Axis, Delta) * SinExcess;
+
+	// Scale the correction by a compliance-derived factor. When compliance is 0 (very high
+	// strength) Alpha is 1 and the correction is fully applied. When compliance is large (low
+	// strength) Alpha approaches 0 and the correction is soft.
+	const float InvMassSum = Particle.InvMass + (bOnlyMoveChild ? 0.0f : ParentParticlePtr->InvMass);
+	// A = alpha~ in the Mueller paper
+	const float A = ParticleTarget.AngleLimitCompliance * FMath::Square(InvDeltaTime);
+	const float Alpha = InvMassSum / (InvMassSum + A);
+	const FSimVector Correction = (NewDelta - Delta) * Alpha;
+
+	if (bOnlyMoveChild)
+	{
+		Particle.Position += Correction;
+	}
+	else
+	{
+		if (InvMassSum > KINDA_SMALL_NUMBER)
+		{
+			Particle.Position += Correction * (Particle.InvMass / InvMassSum);
+			ParentParticlePtr->Position -= Correction * (ParentParticlePtr->InvMass / InvMassSum);
+		}
+	}
+}
+
 //======================================================================================================================
 static void SolveHardDistanceConstraint(
 	const FHardDistanceConstraint& Constraint,
@@ -365,6 +411,103 @@ static void SolveSoftDistanceConstraint(
 	ChildParticle.Position += Correction * InvMassChild;
 }
 
+//======================================================================================================================
+// Enforces the cone limit: the angle between the incoming bone direction (grandparent->parent)
+// and outgoing bone direction (parent->child) must not exceed the limit angle. Uses XPBD with
+// compliance and damping. Corrections are applied to all three particles with mass weighting,
+// and the gradients sum to zero so the centre of mass is preserved.
+static void SolveConeLimit(
+	FConeLimit& ConeLimit,
+	FParticle&  GrandparentParticle,
+	FParticle&  ParentParticle,
+	FParticle&  ChildParticle,
+	const float DeltaTime,
+	const float InvDeltaTime)
+{
+	const float InvMassG = GrandparentParticle.InvMass;
+	const float InvMassP = ParentParticle.InvMass;
+	const float InvMassK = ChildParticle.InvMass;
+
+	if (InvMassG + InvMassP + InvMassK < KINDA_SMALL_NUMBER)
+	{
+		return; // All kinematic
+	}
+
+	// Bone directions: incoming (grandparent->parent), outgoing (parent->child)
+	const FSimVector D1 = ParentParticle.Position - GrandparentParticle.Position;
+	const FSimVector D2 = ChildParticle.Position - ParentParticle.Position;
+	const float R1 = D1.Size();
+	const float R2 = D2.Size();
+
+	if (R1 < KINDA_SMALL_NUMBER || R2 < KINDA_SMALL_NUMBER)
+	{
+		return;
+	}
+
+	const FSimVector N1 = D1 * (1.0f / R1);
+	const FSimVector N2 = D2 * (1.0f / R2);
+
+	const float CosTheta = FMath::Clamp(
+		(float)FSimVector::DotProduct(N1, N2), -1.0f + KINDA_SMALL_NUMBER, 1.0f - KINDA_SMALL_NUMBER);
+	const float SinTheta = FMath::Sqrt(1.0f - CosTheta * CosTheta);
+
+	if (SinTheta < KINDA_SMALL_NUMBER)
+	{
+		return; // Degenerate: nearly straight or fully folded
+	}
+
+	const float Theta = FMath::Acos(CosTheta);
+	const float C = Theta - ConeLimit.Angle;
+
+	if (C <= 0.0f)
+	{
+		return; // Within the limit
+	}
+
+	const float InvSinTheta = 1.0f / SinTheta;
+
+	// Gradients of theta w.r.t. each particle position. These sum to zero, preserving centre
+	// of mass when corrections are applied with inverse-mass weighting.
+	const FSimVector GradG = (N2 - CosTheta * N1) * (InvSinTheta / R1);
+	const FSimVector GradK = -(N1 - CosTheta * N2) * (InvSinTheta / R2);
+	const FSimVector GradP = -(GradG + GradK);
+
+	// Weighted sum of squared gradient norms
+	const float W = InvMassG * GradG.SizeSquared()
+		+ InvMassP * GradP.SizeSquared()
+		+ InvMassK * GradK.SizeSquared();
+
+	if (W < KINDA_SMALL_NUMBER)
+	{
+		return;
+	}
+
+	// XPBD compliance and damping terms
+	// A = alpha~ in the Mueller paper
+	const float A = ConeLimit.Compliance * FMath::Square(InvDeltaTime);
+	// G = gamma in the Mueller paper = alpha~ * beta~ / dt
+	const float G = ConeLimit.Compliance * ConeLimit.Damping * InvDeltaTime;
+
+	// Constraint velocity for damping
+	const FSimVector VelG = (GrandparentParticle.Position - GrandparentParticle.PrevPosition) * InvDeltaTime;
+	const FSimVector VelP = (ParentParticle.Position - ParentParticle.PrevPosition) * InvDeltaTime;
+	const FSimVector VelK = (ChildParticle.Position - ChildParticle.PrevPosition) * InvDeltaTime;
+	const float dCdt = 
+		FSimVector::DotProduct(GradG, VelG) + FSimVector::DotProduct(GradP, VelP) + FSimVector::DotProduct(GradK, VelK);
+
+	// Solve
+	const float Num = -C - A * ConeLimit.Lambda - G * dCdt * DeltaTime;
+	const float Denom = (1.0f + G) * W + A;
+	const float DeltaLambda = Num / Denom;
+
+	ConeLimit.Lambda += DeltaLambda;
+
+	// Apply corrections
+	GrandparentParticle.Position += GradG * (InvMassG * DeltaLambda);
+	ParentParticle.Position += GradP * (InvMassP * DeltaLambda);
+	ChildParticle.Position += GradK * (InvMassK * DeltaLambda);
+}
+
 //======================================================================================================================
 // Note that the docs in the Collider Components note that scale is ignored
 static void SolveContactConstraint(
@@ -434,12 +577,13 @@ static void SolveContactConstraint(
 			FMath::Clamp(LocalPos.Z, -HalfExtents.Z, HalfExtents.Z));
 		const FSimVector ClosestPoint = TM.TransformPositionNoScale(ClampedPos);
 		const FSimVector Delta = Particle.Position - ClosestPoint;
-		const FSimScalar DeltaLen = Delta.Size();
-		if (DeltaLen > ParticleCollider.Radius || DeltaLen < KINDA_SMALL_NUMBER)
+		const FSimScalar DeltaLenSq = Delta.SizeSquared();
+		if (DeltaLenSq > FMath::Square(ParticleCollider.Radius) || DeltaLenSq < KINDA_SMALL_NUMBER)
 		{
 			return;
 		}
 
+		const FSimScalar DeltaLen = FMath::Sqrt(DeltaLenSq);
 		N = Delta * (1.0f / DeltaLen);
 		C = ParticleCollider.Radius - DeltaLen;
 	}
@@ -482,12 +626,13 @@ static void SolveContactConstraint(
 	}
 
 	const FSimVector Delta = Particle.Position - SegmentPoint;
-	const float DeltaLen = Delta.Size();
-	if (DeltaLen > MinDistance || DeltaLen < KINDA_SMALL_NUMBER)
+	const float DeltaLenSq = Delta.SizeSquared();
+	if (DeltaLenSq > FMath::Square(MinDistance) || DeltaLenSq < KINDA_SMALL_NUMBER)
 	{
 		return;
 	}
 
+	const float DeltaLen = FMath::Sqrt(DeltaLenSq);
 	const FSimVector N = Delta * (1.0f / DeltaLen);
 	const float C = MinDistance - DeltaLen;
 
@@ -574,8 +719,8 @@ static void SolveParticleToParticleContact(
 	const FSimScalar CombinedRadius = ParticleACollider.Radius + ParticleBCollider.Radius;
 
 	const FSimVector Delta = ParticleA.Position - ParticleB.Position;
-	FSimScalar DeltaLen = Delta.Size();
-	if (DeltaLen > CombinedRadius || DeltaLen < KINDA_SMALL_NUMBER)
+	const FSimScalar DeltaLenSq = Delta.SizeSquared();
+	if (DeltaLenSq > FMath::Square(CombinedRadius) || DeltaLenSq < KINDA_SMALL_NUMBER)
 	{
 		return;
 	}
@@ -588,6 +733,7 @@ static void SolveParticleToParticleContact(
 		return;
 	}
 
+	const FSimScalar DeltaLen = FMath::Sqrt(DeltaLenSq);
 	const FSimVector N = Delta * (1.0f / DeltaLen);
 	const FSimVector Correction = N * ((CombinedRadius - DeltaLen) / InvMassSum);
 
@@ -621,8 +767,8 @@ void UpdateParticle(
 		? 1.0f / FMath::Max(InMass, KINDA_SMALL_NUMBER)
 		: 0.0f;
 
-	Particle.Update(InvMass);
-	ParticleCollider.Update(InRadius);
+	Particle.InvMass = InvMass;
+	ParticleCollider.Radius = InRadius;
 }
 
 //======================================================================================================================
@@ -635,7 +781,9 @@ void UpdateParticleTarget(
 	const float        InDampingRatio,
 	const float        InExtraDamping,
 	const float        InTargetVelocityInfluence,
-	const float        InTargetMode)
+	const float        InTargetMode,
+	const float        InAngleLimit,
+	const float        InAngleLimitStrength)
 {
 	FParticleInfo& Info = State.ParticleInfos[Index];
 	const FParticle& Particle = State.Particles[Index];
@@ -650,9 +798,10 @@ void UpdateParticleTarget(
 	Info.TargetVelocity = (Info.TargetPosition - Info.PreviousTargetPosition) * InvDeltaTime;
 
 	// Compute and store the normalized direction from parent to this particle in the target pose.
-	// Only needed when TargetMode > 0 (directional contribution).
+	// Needed when TargetMode > 0 (directional contribution) or when angle limits are active.
 	// It is intentional that both directions use the current ParentParticleTargetPosition.
-	if (InTargetMode > KINDA_SMALL_NUMBER
+	const bool bAngleLimitActive = InAngleLimitStrength > KINDA_SMALL_NUMBER;
+	if ((InTargetMode > KINDA_SMALL_NUMBER || bAngleLimitActive)
 		&& State.ParticleInfos.IsValidIndex(Info.ParentParticleIndex))
 	{
 		const FSimVector& ParentTargetPos = State.ParticleInfos[Info.ParentParticleIndex].TargetPosition;
@@ -671,9 +820,33 @@ void UpdateParticleTarget(
 		State.Particles.IsValidIndex(Info.ParentParticleIndex)
 		? &State.Particles[Info.ParentParticleIndex] : nullptr;
 
-	ParticleTarget.Update(
-		Particle, ParentParticle, InStrength, InDampingRatio, InExtraDamping,
-		InTargetVelocityInfluence, InTargetMode);
+	// Angle limit (compliance from strength, same pattern as the other constraints)
+	if (bAngleLimitActive && ParentParticle)
+	{
+		const float AngleLimitRadians = FMath::DegreesToRadians(FMath::Max(InAngleLimit, 0.0f));
+		FMath::SinCos(&ParticleTarget.SinAngleLimit, &ParticleTarget.CosAngleLimit, AngleLimitRadians);
+
+		const float W = InAngleLimitStrength * TWO_PI;
+		const float InvMassSum = Particle.InvMass + ParentParticle->InvMass;
+		ParticleTarget.AngleLimitCompliance = InvMassSum / FMath::Max(W * W, KINDA_SMALL_NUMBER);
+	}
+	else
+	{
+		ParticleTarget.CosAngleLimit = 1.0f;
+		ParticleTarget.SinAngleLimit = 0.0f;
+		ParticleTarget.AngleLimitCompliance = 0.0f;
+	}
+
+	ParticleTarget.TargetVelocityInfluence = InTargetVelocityInfluence;
+	ParticleTarget.TargetMode = InTargetMode;
+
+	// Scale compliance by the effective inverse mass so the constraint frequency is
+	// mass-independent. Parent contribution is weighted by Alpha^2 (from the squared Jacobian).
+	const float Alpha = ParentParticle ? InTargetMode : 0.0f;
+	const float InvMass = Particle.InvMass + Alpha * Alpha * (ParentParticle ? ParentParticle->InvMass : 0.0f);
+	const float W = InStrength * TWO_PI;
+	ParticleTarget.Compliance = InvMass / FMath::Max(W * W, KINDA_SMALL_NUMBER);
+	ParticleTarget.Damping = (2.0f * InDampingRatio * W + InExtraDamping) / FMath::Max(InvMass, KINDA_SMALL_NUMBER);
 }
 
 //======================================================================================================================
@@ -702,6 +875,11 @@ void UpdateDynamicsStep(
 		StepSoftDistanceConstraint(Constraint);
 	}
 
+	for (FConeLimit& ConeLimit : SimulationState.ConeLimits)
+	{
+		StepConeLimit(ConeLimit);
+	}
+
 	// Hard distance constraints don't have a step
 
 	for (FCollider& Collider : SimulationState.Colliders)
@@ -746,6 +924,12 @@ void UpdateDynamicsStep(
 				SolveTargetConstraint(
 					Particle, ParticleTarget, ParentParticlePtr, bTargetOnlyMoveChild,
 					DeltaTime, InvDeltaTime);
+				if (ParticleTarget.AngleLimitCompliance > 0.0f)
+				{
+					SolveTargetAngleConstraint(
+						Particle, ParticleTarget, ParentParticlePtr, bTargetOnlyMoveChild,
+						InvDeltaTime);
+				}
 			}
 		}
 
@@ -770,6 +954,24 @@ void UpdateDynamicsStep(
 					InvDeltaTime);
 			}
 
+			// Cone limits (angular constraints between triples of particles)
+			for (int32 Index = 0; Index != SimulationState.ConeLimits.Num(); ++Index)
+			{
+				const FConeLimitInfo& Info = SimulationState.ConeLimitInfos[Index];
+				FConeLimit& ConeLimit = SimulationState.ConeLimits[Index];
+
+				check(SimulationState.Particles.IsValidIndex(Info.GrandparentIndex) &&
+					SimulationState.Particles.IsValidIndex(Info.ParentIndex) &&
+					SimulationState.Particles.IsValidIndex(Info.ChildIndex));
+				SolveConeLimit(
+					ConeLimit,
+					SimulationState.Particles[Info.GrandparentIndex],
+					SimulationState.Particles[Info.ParentIndex],
+					SimulationState.Particles[Info.ChildIndex],
+					DeltaTime,
+					InvDeltaTime);
+			}
+
 			// Hard distance constraints come later to reduce stretching
 			for (int32 Index = 0; Index != SimulationState.HardDistanceConstraints.Num(); ++Index)
 			{
@@ -977,6 +1179,12 @@ TArray<int32> SortParticlesRootToLeaf(FSimulationState& State)
 		Info.ParentIndex = OldToNew[Info.ParentIndex];
 		Info.ChildIndex = OldToNew[Info.ChildIndex];
 	}
+	for (FConeLimitInfo& Info : State.ConeLimitInfos)
+	{
+		Info.GrandparentIndex = OldToNew[Info.GrandparentIndex];
+		Info.ParentIndex = OldToNew[Info.ParentIndex];
+		Info.ChildIndex = OldToNew[Info.ChildIndex];
+	}
 
 	return SortedOrder;
 }
diff --git a/Engine/Plugins/Experimental/ControlRigDynamics/Source/ControlRigDynamics/Public/RigDynamicsConeLimitComponent.h b/Engine/Plugins/Experimental/ControlRigDynamics/Source/ControlRigDynamics/Public/RigDynamicsConeLimitComponent.h
new file mode 100644
--- /dev/null
+++ b/Engine/Plugins/Experimental/ControlRigDynamics/Source/ControlRigDynamics/Public/RigDynamicsConeLimitComponent.h
@@ -0,0 +1,62 @@
+// Copyright Epic Games, Inc. All Rights Reserved.
+
+#pragma once
+
+#include "RigParticleSimulation.h"
+
+#include "RigDynamicsParticleComponent.h"
+
+#include "RigDynamicsConeLimitComponent.generated.h"
+
+#define UE_API CONTROLRIGDYNAMICS_API
+
+//======================================================================================================================
+// A component that can be added to a joint/element that defines a cone (angular) limit between three particles. 
+//======================================================================================================================
+USTRUCT(BlueprintType)
+struct FRigDynamicsConeLimitComponent : public FRigBaseComponent
+{
+public:
+	GENERATED_BODY()
+	DECLARE_RIG_COMPONENT_METHODS(FRigDynamicsConeLimitComponent)
+
+	// The grandparent particle for the ConeLimit
+	UPROPERTY(BlueprintReadOnly, EditAnywhere, Category = ConeLimit)
+	FRigComponentKey GrandparentComponentKey = FRigComponentKey(
+		FRigElementKey(TEXT(""), ERigElementType::Bone), FRigDynamicsParticleComponent::GetDefaultName());
+
+	// The parent particle for the ConeLimit
+	UPROPERTY(BlueprintReadOnly, EditAnywhere, Category = ConeLimit)
+	FRigComponentKey ParentComponentKey = FRigComponentKey(
+		FRigElementKey(TEXT(""), ERigElementType::Bone), FRigDynamicsParticleComponent::GetDefaultName());
+
+	// The child particle for the ConeLimit
+	UPROPERTY(BlueprintReadOnly, EditAnywhere, Category = ConeLimit)
+	FRigComponentKey ChildComponentKey = FRigComponentKey(
+		FRigElementKey(TEXT(""), ERigElementType::Bone), FRigDynamicsParticleComponent::GetDefaultName());
+
+	// The strength which we drive towards the target limit (when outside). This is the oscillation
+	// frequency, so low values will be soft and springy, but values significantly above 1/timestep
+	// will impose the limit strongly.
+	UPROPERTY(BlueprintReadOnly, EditAnywhere, Category = Dynamics, meta = (ClampMin = "0.0"))
+	float Strength = 10.0f;
+
+	// DampingRatio for the ConeLimit
+	UPROPERTY(BlueprintReadOnly, EditAnywhere, Category = Dynamics, meta = (ClampMin = "0.0"))
+	float DampingRatio = 1.0f;
+
+	// The (full) cone angle in degrees
+	UPROPERTY(BlueprintReadOnly, EditAnywhere, Category = Dynamics, meta = (ClampMin = "0.0"))
+	float Angle = 0.0f;
+
+	UE_API virtual void Save(FArchive& Ar) override;
+	UE_API virtual void Load(FArchive& Ar) override;
+
+	virtual FName GetDefaultComponentName() const override { return GetDefaultName(); }
+	static FName GetDefaultName() { return TEXT("DynamicsConeLimit"); }
+
+protected:
+	void Serialize(FArchive& Ar);
+};
+
+#undef UE_API
diff --git a/Engine/Plugins/Experimental/ControlRigDynamics/Source/ControlRigDynamics/Public/RigDynamicsConeLimitExecution.h b/Engine/Plugins/Experimental/ControlRigDynamics/Source/ControlRigDynamics/Public/RigDynamicsConeLimitExecution.h
new file mode 100644
--- /dev/null
+++ b/Engine/Plugins/Experimental/ControlRigDynamics/Source/ControlRigDynamics/Public/RigDynamicsConeLimitExecution.h
@@ -0,0 +1,192 @@
+// Copyright Epic Games, Inc. All Rights Reserved.
+
+#pragma once
+
+#include "RigDynamicsExecution.h"
+#include "RigDynamicsConeLimitComponent.h"
+#include "RigDynamicsSolverComponent.h"
+#include "RigDynamicsParticleComponent.h"
+
+#include "RigParticleSimulation.h"
+
+#include "RigDynamicsConeLimitExecution.generated.h"
+
+#define UE_API CONTROLRIGDYNAMICS_API
+
+//======================================================================================================================
+// Spawns a new dynamics cone limit component.
+// Note: This node only runs as part of the construction event.
+//======================================================================================================================
+USTRUCT(meta = (DisplayName = "Spawn Dynamics Cone Limit", Keywords = "Add,Construction,Create,Angle,Angular", Varying))
+struct FRigUnit_SpawnDynamicsConeLimit : public FRigUnit_DynamicsBaseMutable
+{
+	GENERATED_BODY()
+
+	RIGVM_METHOD()
+	UE_API virtual void Execute() override;
+
+	// The owner of the newly created component (must be set/valid)
+	UPROPERTY(meta = (Input, BoneName))
+	FRigElementKey Owner = FRigElementKey(TEXT(""), ERigElementType::Bone);
+
+	// Optional solver - if set, it will be added to this solver component
+	UPROPERTY(meta = (Input))
+	FRigComponentKey DynamicsSolverComponentKey = FRigComponentKey(
+		FRigElementKey(TEXT(""), ERigElementType::Bone), FRigDynamicsSolverComponent::GetDefaultName());
+
+	// The grandparent particle for the cone limit
+	UPROPERTY(meta = (Input))
+	FRigComponentKey GrandparentComponentKey = FRigComponentKey(
+		FRigElementKey(TEXT(""), ERigElementType::Bone), FRigDynamicsParticleComponent::GetDefaultName());
+
+	// The parent particle for the cone limit
+	UPROPERTY(meta = (Input))
+	FRigComponentKey ParentComponentKey = FRigComponentKey(
+		FRigElementKey(TEXT(""), ERigElementType::Bone), FRigDynamicsParticleComponent::GetDefaultName());
+
+	// The child particle for the cone limit
+	UPROPERTY(meta = (Input))
+	FRigComponentKey ChildComponentKey = FRigComponentKey(
+		FRigElementKey(TEXT(""), ERigElementType::Bone), FRigDynamicsParticleComponent::GetDefaultName());
+
+	// The strength which we drive towards the target limit (when outside). This is the oscillation
+	// frequency, so low values will be soft and springy, but values significantly above 1/timestep
+	// will impose the limit strongly.
+	UPROPERTY(meta = (Input, ClampMin = "0.0"))
+	float Strength = 10.0f;
+
+	// DampingRatio for the cone limit
+	UPROPERTY(meta = (Input, ClampMin = "0.0"))
+	float DampingRatio = 1.0f;
+
+	// The (full) cone angle in degrees
+	UPROPERTY(meta = (Input, ClampMin = "0.0"))
+	float Angle = 0.0f;
+
+	// The Dynamics Cone Limit component key that was created
+	UPROPERTY(meta = (Output))
+	FRigComponentKey DynamicsConeLimitComponentKey;
+};
+
+//======================================================================================================================
+// Sets the strength of a Dynamics cone limit
+USTRUCT(meta = (DisplayName = "Set Dynamics Cone Limit Strength", Varying))
+struct FRigUnit_HierarchySetDynamicsConeLimitStrength : public FRigUnit_DynamicsBaseMutable
+{
+	GENERATED_BODY()
+
+	RIGVM_METHOD()
+	UE_API virtual void Execute() override;
+
+	// The Dynamics Cone Limit to be updated
+	UPROPERTY(meta = (Input))
+	FRigComponentKey DynamicsConeLimitComponentKey = FRigComponentKey(
+		FRigElementKey(TEXT(""), ERigElementType::Bone), FRigDynamicsConeLimitComponent::GetDefaultName());
+
+	// The strength which enforces the limit (oscillation frequency)
+	UPROPERTY(meta = (Input, ClampMin = "0.0"))
+	float Strength = 10.0f;
+};
+
+//======================================================================================================================
+// Gets the strength of a Dynamics cone limit
+USTRUCT(meta = (DisplayName = "Get Dynamics Cone Limit Strength", Varying))
+struct FRigUnit_HierarchyGetDynamicsConeLimitStrength : public FRigUnit_DynamicsBase
+{
+	GENERATED_BODY()
+
+	RIGVM_METHOD()
+	UE_API virtual void Execute() override;
+
+	// The Dynamics Cone Limit to query
+	UPROPERTY(meta = (Input))
+	FRigComponentKey DynamicsConeLimitComponentKey = FRigComponentKey(
+		FRigElementKey(TEXT(""), ERigElementType::Bone), FRigDynamicsConeLimitComponent::GetDefaultName());
+
+	// The strength which enforces the limit (oscillation frequency)
+	UPROPERTY(meta = (Output))
+	float Strength = 0.0f;
+};
+
+//======================================================================================================================
+// Sets the damping ratio of a Dynamics cone limit
+USTRUCT(meta = (DisplayName = "Set Dynamics Cone Limit Damping Ratio", Varying))
+struct FRigUnit_HierarchySetDynamicsConeLimitDampingRatio : public FRigUnit_DynamicsBaseMutable
+{
+	GENERATED_BODY()
+
+	RIGVM_METHOD()
+	UE_API virtual void Execute() override;
+
+	// The Dynamics Cone Limit to be updated
+	UPROPERTY(meta = (Input))
+	FRigComponentKey DynamicsConeLimitComponentKey = FRigComponentKey(
+		FRigElementKey(TEXT(""), ERigElementType::Bone), FRigDynamicsConeLimitComponent::GetDefaultName());
+
+	// Damping ratio for the cone limit
+	UPROPERTY(meta = (Input, ClampMin = "0.0"))
+	float DampingRatio = 1.0f;
+};
+
+//======================================================================================================================
+// Gets the damping ratio of a Dynamics cone limit
+USTRUCT(meta = (DisplayName = "Get Dynamics Cone Limit Damping Ratio", Varying))
+struct FRigUnit_HierarchyGetDynamicsConeLimitDampingRatio : public FRigUnit_DynamicsBase
+{
+	GENERATED_BODY()
+
+	RIGVM_METHOD()
+	UE_API virtual void Execute() override;
+
+	// The Dynamics Cone Limit to query
+	UPROPERTY(meta = (Input))
+	FRigComponentKey DynamicsConeLimitComponentKey = FRigComponentKey(
+		FRigElementKey(TEXT(""), ERigElementType::Bone), FRigDynamicsConeLimitComponent::GetDefaultName());
+
+	// Damping ratio for the cone limit
+	UPROPERTY(meta = (Output))
+	float DampingRatio = 0.0f;
+};
+
+//======================================================================================================================
+// Sets the angle of a Dynamics cone limit
+USTRUCT(meta = (DisplayName = "Set Dynamics Cone Limit Angle", Varying))
+struct FRigUnit_HierarchySetDynamicsConeLimitAngle : public FRigUnit_DynamicsBaseMutable
+{
+	GENERATED_BODY()
+
+	RIGVM_METHOD()
+	UE_API virtual void Execute() override;
+
+	// The Dynamics Cone Limit to be updated
+	UPROPERTY(meta = (Input))
+	FRigComponentKey DynamicsConeLimitComponentKey = FRigComponentKey(
+		FRigElementKey(TEXT(""), ERigElementType::Bone), FRigDynamicsConeLimitComponent::GetDefaultName());
+
+	// The (full) cone angle in degrees
+	UPROPERTY(meta = (Input, ClampMin = "0.0"))
+	float Angle = 0.0f;
+};
+
+//======================================================================================================================
+// Gets the angle of a Dynamics cone limit
+USTRUCT(meta = (DisplayName = "Get Dynamics Cone Limit Angle", Varying))
+struct FRigUnit_HierarchyGetDynamicsConeLimitAngle : public FRigUnit_DynamicsBase
+{
+	GENERATED_BODY()
+
+	RIGVM_METHOD()
+	UE_API virtual void Execute() override;
+
+	// The Dynamics Cone Limit to query
+	UPROPERTY(meta = (Input))
+	FRigComponentKey DynamicsConeLimitComponentKey = FRigComponentKey(
+		FRigElementKey(TEXT(""), ERigElementType::Bone), FRigDynamicsConeLimitComponent::GetDefaultName());
+
+	// The (full) cone angle in degrees
+	UPROPERTY(meta = (Output))
+	float Angle = 0.0f;
+};
+
+
+#undef UE_API
diff --git a/Engine/Plugins/Experimental/ControlRigDynamics/Source/ControlRigDynamics/Public/RigDynamicsConstraintComponent.h b/Engine/Plugins/Experimental/ControlRigDynamics/Source/ControlRigDynamics/Public/RigDynamicsConstraintComponent.h
--- a/Engine/Plugins/Experimental/ControlRigDynamics/Source/ControlRigDynamics/Public/RigDynamicsConstraintComponent.h
+++ b/Engine/Plugins/Experimental/ControlRigDynamics/Source/ControlRigDynamics/Public/RigDynamicsConstraintComponent.h
@@ -4,7 +4,7 @@
 
 #include "RigParticleSimulation.h"
 
-#include "Rigs/RigHierarchyComponents.h"
+#include "RigDynamicsParticleComponent.h"
 
 #include "RigDynamicsConstraintComponent.generated.h"
 
@@ -22,11 +22,13 @@ public:
 
 	// The parent particle for the constraint
 	UPROPERTY(BlueprintReadOnly, EditAnywhere, Category = Constraint)
-	FRigComponentKey ParentComponentKey;
+	FRigComponentKey ParentComponentKey = FRigComponentKey(
+		FRigElementKey(TEXT(""), ERigElementType::Bone), FRigDynamicsParticleComponent::GetDefaultName());
 
 	// The child particle for the constraint
 	UPROPERTY(BlueprintReadOnly, EditAnywhere, Category = Constraint)
-	FRigComponentKey ChildComponentKey;
+	FRigComponentKey ChildComponentKey = FRigComponentKey(
+		FRigElementKey(TEXT(""), ERigElementType::Bone), FRigDynamicsParticleComponent::GetDefaultName());
 
 	// If the constraint is marked as "hard", then the system will try as hard as possible to
 	// enforce it. This will also be a little cheaper than making a soft constraint.
diff --git a/Engine/Plugins/Experimental/ControlRigDynamics/Source/ControlRigDynamics/Public/RigDynamicsConstraintExecution.h b/Engine/Plugins/Experimental/ControlRigDynamics/Source/ControlRigDynamics/Public/RigDynamicsConstraintExecution.h
--- a/Engine/Plugins/Experimental/ControlRigDynamics/Source/ControlRigDynamics/Public/RigDynamicsConstraintExecution.h
+++ b/Engine/Plugins/Experimental/ControlRigDynamics/Source/ControlRigDynamics/Public/RigDynamicsConstraintExecution.h
@@ -53,15 +53,15 @@ struct FRigUnit_SpawnDynamicsConstraint : public FRigUnit_DynamicsBaseMutable
 	// low values will be soft and springy, but values significantly above 1/timestep will track
 	// the target very accurately.
 	UPROPERTY(meta = (Input, ClampMin = "0.0", EditCondition = "ConstraintType != ERigParticleSimulationConstraintType::Hard"))
-	float Strength = 2.0;
+	float Strength = 2.0f;
 
 	// DampingRatio for the constraint
 	UPROPERTY(meta = (Input, ClampMin = "0.0", EditCondition = "ConstraintType != ERigParticleSimulationConstraintType::Hard"))
-	float DampingRatio = 1.0;
+	float DampingRatio = 1.0f;
 
 	// Additional damping for the constraint
 	UPROPERTY(meta = (Input, ClampMin = "0.0", EditCondition = "ConstraintType != ERigParticleSimulationConstraintType::Hard"))
-	float ExtraDamping = 0.0;
+	float ExtraDamping = 0.0f;
 
 	// The target length will be calculated automatically, multiplied by this
 	UPROPERTY(meta = (Input, ClampMin = "0.0"))
@@ -93,7 +93,7 @@ struct FRigUnit_HierarchySetDynamicsConstraintStrength : public FRigUnit_Dynamic
 
 	// The strength which drives towards the target length (oscillation frequency)
 	UPROPERTY(meta = (Input, ClampMin = "0.0"))
-	float Strength = 2.0;
+	float Strength = 2.0f;
 };
 
 //======================================================================================================================
@@ -133,7 +133,7 @@ struct FRigUnit_HierarchySetDynamicsConstraintDampingRatio : public FRigUnit_Dyn
 
 	// Damping ratio for the constraint
 	UPROPERTY(meta = (Input, ClampMin = "0.0"))
-	float DampingRatio = 1.0;
+	float DampingRatio = 1.0f;
 };
 
 //======================================================================================================================
@@ -173,7 +173,7 @@ struct FRigUnit_HierarchySetDynamicsConstraintExtraDamping : public FRigUnit_Dyn
 
 	// Additional damping for the constraint
 	UPROPERTY(meta = (Input, ClampMin = "0.0"))
-	float ExtraDamping = 0.0;
+	float ExtraDamping = 0.0f;
 };
 
 //======================================================================================================================
diff --git a/Engine/Plugins/Experimental/ControlRigDynamics/Source/ControlRigDynamics/Public/RigDynamicsData.h b/Engine/Plugins/Experimental/ControlRigDynamics/Source/ControlRigDynamics/Public/RigDynamicsData.h
--- a/Engine/Plugins/Experimental/ControlRigDynamics/Source/ControlRigDynamics/Public/RigDynamicsData.h
+++ b/Engine/Plugins/Experimental/ControlRigDynamics/Source/ControlRigDynamics/Public/RigDynamicsData.h
@@ -326,19 +326,19 @@ struct FRigDynamicsVisualizationSettings
 	// Enable visualization. There is a significant runtime cost to enabling visualization, so make
 	// sure you disable it when it's not needed, and when you are profiling. Note that this can be
 	// overridden using CVar ControlRig.Dynamics.VisualizationOverride
-	UPROPERTY(EditAnywhere, Category = Visualization)
+	UPROPERTY(BlueprintReadOnly, EditAnywhere, Category = Visualization)
 	bool bEnableVisualization = false;
 
-	UPROPERTY(EditAnywhere, Category = Visualization, meta = (ClampMin = "0.0"))
+	UPROPERTY(BlueprintReadOnly, EditAnywhere, Category = Visualization, meta = (ClampMin = "0.0"))
 	float LineThickness = 1.0f;
 
-	UPROPERTY(EditAnywhere, Category = Visualization, meta = (ClampMin = "0.0"))
+	UPROPERTY(BlueprintReadOnly, EditAnywhere, Category = Visualization, meta = (ClampMin = "6"))
 	int32 ShapeDetail = 16;
 
-	UPROPERTY(EditAnywhere, Category = Visualization)
+	UPROPERTY(BlueprintReadOnly, EditAnywhere, Category = Visualization)
 	bool bShowParticles = true;
 
-	UPROPERTY(EditAnywhere, Category = Visualization)
+	UPROPERTY(BlueprintReadOnly, EditAnywhere, Category = Visualization)
 	bool bShowColliders = true;
 };
 
diff --git a/Engine/Plugins/Experimental/ControlRigDynamics/Source/ControlRigDynamics/Public/RigDynamicsParticleComponent.h b/Engine/Plugins/Experimental/ControlRigDynamics/Source/ControlRigDynamics/Public/RigDynamicsParticleComponent.h
--- a/Engine/Plugins/Experimental/ControlRigDynamics/Source/ControlRigDynamics/Public/RigDynamicsParticleComponent.h
+++ b/Engine/Plugins/Experimental/ControlRigDynamics/Source/ControlRigDynamics/Public/RigDynamicsParticleComponent.h
@@ -60,6 +60,17 @@ struct FRigDynamicsParticleProperties
 	UPROPERTY(BlueprintReadOnly, EditAnywhere, Category = Dynamics, meta = (ClampMin = "0.0", ClampMax = "1.0"))
 	float TargetMode = 0.5f;
 
+	// Constrains this particle to align with the target direction from its parent, to within this
+	// angle (in degrees). This effectively gives us a limit on the deviation from the target pose.
+	UPROPERTY(BlueprintReadOnly, EditAnywhere, Category = Dynamics, meta = (ClampMin = "0.0"))
+	float AngleLimit = 0.0f;
+
+	// The strength of the angle limit constraint (oscillation frequency). High values impose the
+	// limit more rigidly, low values allow soft/springy deviation before settling. Zero disables
+	// the angle limit.
+	UPROPERTY(BlueprintReadOnly, EditAnywhere, Category = Dynamics, meta = (ClampMin = "0.0"))
+	float AngleLimitStrength = 0.0f;
+
 	// A list of colliders we shouldn't collide with
 	UPROPERTY(BlueprintReadOnly, EditAnywhere, Category = Dynamics)
 	TArray<FRigComponentKey> NoCollisionColliders;
diff --git a/Engine/Plugins/Experimental/ControlRigDynamics/Source/ControlRigDynamics/Public/RigDynamicsParticleExecution.h b/Engine/Plugins/Experimental/ControlRigDynamics/Source/ControlRigDynamics/Public/RigDynamicsParticleExecution.h
--- a/Engine/Plugins/Experimental/ControlRigDynamics/Source/ControlRigDynamics/Public/RigDynamicsParticleExecution.h
+++ b/Engine/Plugins/Experimental/ControlRigDynamics/Source/ControlRigDynamics/Public/RigDynamicsParticleExecution.h
@@ -266,7 +266,7 @@ struct FRigUnit_HierarchySetDynamicsParticleStrength : public FRigUnit_DynamicsB
 
 	// The strength which drives towards the target (oscillation frequency)
 	UPROPERTY(meta = (Input, ClampMin = "0.0"))
-	float Strength = 2.0;
+	float Strength = 2.0f;
 };
 
 //======================================================================================================================
@@ -306,7 +306,7 @@ struct FRigUnit_HierarchySetDynamicsParticleDampingRatio : public FRigUnit_Dynam
 
 	// Damping ratio for the particle
 	UPROPERTY(meta = (Input, ClampMin = "0.0"))
-	float DampingRatio = 1.0;
+	float DampingRatio = 1.0f;
 };
 
 //======================================================================================================================
@@ -346,7 +346,7 @@ struct FRigUnit_HierarchySetDynamicsParticleExtraDamping : public FRigUnit_Dynam
 
 	// Extra damping for the particle
 	UPROPERTY(meta = (Input, ClampMin = "0.0"))
-	float ExtraDamping = 0.0;
+	float ExtraDamping = 0.0f;
 };
 
 //======================================================================================================================
@@ -386,7 +386,7 @@ struct FRigUnit_HierarchySetDynamicsParticleTargetVelocityInfluence : public FRi
 
 	// If 1 the target's velocity will be used for damping, if 0 damping is in sim space
 	UPROPERTY(meta = (Input, ClampMin = "0.0"))
-	float TargetVelocityInfluence = 1.0;
+	float TargetVelocityInfluence = 1.0f;
 };
 
 //======================================================================================================================
@@ -449,4 +449,92 @@ struct FRigUnit_HierarchyGetDynamicsParticleTargetMode : public FRigUnit_Dynamic
 	float TargetMode = 0.5f;
 };
 
+//======================================================================================================================
+// Sets the angle limit of a Dynamics particle
+USTRUCT(meta = (DisplayName = "Set Dynamics Particle Angle Limit", Varying))
+struct FRigUnit_HierarchySetDynamicsParticleAngleLimit : public FRigUnit_DynamicsBaseMutable
+{
+	GENERATED_BODY()
+
+	RIGVM_METHOD()
+	UE_API virtual void Execute() override;
+
+	// The Dynamics Particle to be updated
+	UPROPERTY(meta = (Input))
+	FRigComponentKey DynamicsParticleComponentKey = FRigComponentKey(
+		FRigElementKey(TEXT(""), ERigElementType::Bone), FRigDynamicsParticleComponent::GetDefaultName());
+
+	// This applies a soft constraint to make this particle align with the target direction from its
+	// parent, to within this angle (in degrees). This effectively gives us a limit on the deviation
+	// from the target pose.
+	UPROPERTY(meta = (Input, ClampMin = "0.0"))
+	float AngleLimit = 0.0f;
+};
+
+//======================================================================================================================
+// Gets the angle limit of a Dynamics particle
+USTRUCT(meta = (DisplayName = "Get Dynamics Particle Angle Limit", Varying))
+struct FRigUnit_HierarchyGetDynamicsParticleAngleLimit : public FRigUnit_DynamicsBase
+{
+	GENERATED_BODY()
+
+	RIGVM_METHOD()
+	UE_API virtual void Execute() override;
+
+	// The Dynamics Particle to query
+	UPROPERTY(meta = (Input))
+	FRigComponentKey DynamicsParticleComponentKey = FRigComponentKey(
+		FRigElementKey(TEXT(""), ERigElementType::Bone), FRigDynamicsParticleComponent::GetDefaultName());
+
+	// This applies a soft constraint to make this particle align with the target direction from its
+	// parent, to within this angle (in degrees). This effectively gives us a limit on the deviation
+	// from the target pose.
+	UPROPERTY(meta = (Output))
+	float AngleLimit = 0.0f;
+};
+
+//======================================================================================================================
+// Sets the angle limit strength of a Dynamics particle
+USTRUCT(meta = (DisplayName = "Set Dynamics Particle Angle Limit Strength", Varying))
+struct FRigUnit_HierarchySetDynamicsParticleAngleLimitStrength : public FRigUnit_DynamicsBaseMutable
+{
+	GENERATED_BODY()
+
+	RIGVM_METHOD()
+	UE_API virtual void Execute() override;
+
+	// The Dynamics Particle to be updated
+	UPROPERTY(meta = (Input))
+	FRigComponentKey DynamicsParticleComponentKey = FRigComponentKey(
+		FRigElementKey(TEXT(""), ERigElementType::Bone), FRigDynamicsParticleComponent::GetDefaultName());
+
+	// The strength of the angle limit constraint (oscillation frequency). High values impose the
+	// limit more rigidly, low values allow soft/springy deviation before settling. Zero disables
+	// the angle limit.
+	UPROPERTY(meta = (Input, ClampMin = "0.0"))
+	float AngleLimitStrength = 0.0f;
+};
+
+//======================================================================================================================
+// Gets the angle limit strength of a Dynamics particle
+USTRUCT(meta = (DisplayName = "Get Dynamics Particle Angle Limit Strength", Varying))
+struct FRigUnit_HierarchyGetDynamicsParticleAngleLimitStrength : public FRigUnit_DynamicsBase
+{
+	GENERATED_BODY()
+
+	RIGVM_METHOD()
+	UE_API virtual void Execute() override;
+
+	// The Dynamics Particle to query
+	UPROPERTY(meta = (Input))
+	FRigComponentKey DynamicsParticleComponentKey = FRigComponentKey(
+		FRigElementKey(TEXT(""), ERigElementType::Bone), FRigDynamicsParticleComponent::GetDefaultName());
+
+	// The strength of the angle limit constraint (oscillation frequency). High values impose the
+	// limit more rigidly, low values allow soft/springy deviation before settling. Zero disables
+	// the angle limit.
+	UPROPERTY(meta = (Output))
+	float AngleLimitStrength = 0.0f;
+};
+
 #undef UE_API
diff --git a/Engine/Plugins/Experimental/ControlRigDynamics/Source/ControlRigDynamics/Public/RigDynamicsSimulationSpace.h b/Engine/Plugins/Experimental/ControlRigDynamics/Source/ControlRigDynamics/Public/RigDynamicsSimulationSpace.h
--- a/Engine/Plugins/Experimental/ControlRigDynamics/Source/ControlRigDynamics/Public/RigDynamicsSimulationSpace.h
+++ b/Engine/Plugins/Experimental/ControlRigDynamics/Source/ControlRigDynamics/Public/RigDynamicsSimulationSpace.h
@@ -43,6 +43,11 @@ public:
 	// Convert a direction/velocity vector from world space into simulation space.
 	FVector ConvertWorldVectorToSimSpace(const FVector& WorldVector) const;
 
+private:
+
+	// Called when the component/simulation TMs have been updated
+	void RecomputeCompositeTransforms();
+
 	// The world-space component (owner of the simulation) transform
 	FTransform ComponentTM;
 
@@ -61,6 +66,11 @@ public:
 	FTransform PrevSimulationSpaceTM;
 	FTransform PrevPrevSimulationSpaceTM;
 
+	// Precomputed composite transforms so that each coordinate conversion is a single transform
+	// operation instead of two. Recomputed whenever ComponentTM or SimulationSpaceTM changes.
+	FTransform ComponentToSimSpaceTM;
+	FTransform SimToComponentSpaceTM;
+
 	// The time between SimulationSpaceTM and PrevSimulationSpaceTM
 	float DeltaTime = 0.0f;
 	// The time between PrevSimulationSpaceTM and PrevPrevSimulationSpaceTM
@@ -76,36 +86,39 @@ public:
 // Inline implementations
 //======================================================================================================================
 
+//======================================================================================================================
 inline FVector FRigDynamicsSimulationSpaceState::ConvertComponentSpacePositionToSimSpace(
 	const FVector& ComponentSpacePosition) const
 {
-	return SimulationSpaceTM.InverseTransformPositionNoScale(
-		ComponentTM.TransformPositionNoScale(ComponentSpacePosition));
+	return ComponentToSimSpaceTM.TransformPositionNoScale(ComponentSpacePosition);
 }
 
+//======================================================================================================================
 inline FVector FRigDynamicsSimulationSpaceState::ConvertSimSpacePositionToComponentSpace(
 	const FVector& SimSpacePosition) const
 {
-	return ComponentTM.InverseTransformPositionNoScale(
-		SimulationSpaceTM.TransformPositionNoScale(SimSpacePosition));
+	return SimToComponentSpaceTM.TransformPositionNoScale(SimSpacePosition);
 }
 
+//======================================================================================================================
 inline FTransform FRigDynamicsSimulationSpaceState::ConvertComponentSpaceTransformToSimSpace(
 	const FTransform& ComponentSpaceTransform) const
 {
-	FTransform Result = ComponentSpaceTransform * ComponentTM * SimulationSpaceTM.Inverse();
+	FTransform Result = ComponentSpaceTransform * ComponentToSimSpaceTM;
 	Result.SetScale3D(FVector::OneVector);
 	return Result;
 }
 
+//======================================================================================================================
 inline FTransform FRigDynamicsSimulationSpaceState::ConvertSimSpaceTransformToComponentSpace(
 	const FTransform& SimSpaceTransform) const
 {
-	FTransform Result = SimSpaceTransform * SimulationSpaceTM * ComponentTM.Inverse();
+	FTransform Result = SimSpaceTransform * SimToComponentSpaceTM;
 	Result.SetScale3D(FVector::OneVector);
 	return Result;
 }
 
+//======================================================================================================================
 inline FVector FRigDynamicsSimulationSpaceState::ConvertWorldVectorToSimSpace(
 	const FVector& WorldVector) const
 {
diff --git a/Engine/Plugins/Experimental/ControlRigDynamics/Source/ControlRigDynamics/Public/RigDynamicsSolverComponent.h b/Engine/Plugins/Experimental/ControlRigDynamics/Source/ControlRigDynamics/Public/RigDynamicsSolverComponent.h
--- a/Engine/Plugins/Experimental/ControlRigDynamics/Source/ControlRigDynamics/Public/RigDynamicsSolverComponent.h
+++ b/Engine/Plugins/Experimental/ControlRigDynamics/Source/ControlRigDynamics/Public/RigDynamicsSolverComponent.h
@@ -39,6 +39,10 @@ public:
 	UPROPERTY(BlueprintReadOnly, EditAnywhere, Category = Components)
 	TArray<FRigComponentKey> Constraints;
 
+	// Cone limits apply forces to triples of particles to limit the angle between them
+	UPROPERTY(BlueprintReadOnly, EditAnywhere, Category = Components)
+	TArray<FRigComponentKey> ConeLimits;
+
 	UPROPERTY(BlueprintReadOnly, EditAnywhere, Category = Solver)
 	FRigDynamicsSolverSettings Settings;
 
diff --git a/Engine/Plugins/Experimental/ControlRigDynamics/Source/ControlRigDynamics/Public/RigDynamicsSolverExecution.h b/Engine/Plugins/Experimental/ControlRigDynamics/Source/ControlRigDynamics/Public/RigDynamicsSolverExecution.h
--- a/Engine/Plugins/Experimental/ControlRigDynamics/Source/ControlRigDynamics/Public/RigDynamicsSolverExecution.h
+++ b/Engine/Plugins/Experimental/ControlRigDynamics/Source/ControlRigDynamics/Public/RigDynamicsSolverExecution.h
@@ -159,7 +159,7 @@ struct FRigUnit_HierarchySetDynamicsSolverGravity : public FRigUnit_DynamicsBase
 
 	// Gravity in world space, can be scaled per particle
 	UPROPERTY(meta = (Input))
-	FVector Gravity = FVector(0, 0, -981.0);
+	FVector Gravity = FVector(0.0f, 0.0f, -981.0f);
 };
 
 //======================================================================================================================
diff --git a/Engine/Plugins/Experimental/ControlRigDynamics/Source/ControlRigDynamics/Public/RigParticleSimulation.h b/Engine/Plugins/Experimental/ControlRigDynamics/Source/ControlRigDynamics/Public/RigParticleSimulation.h
--- a/Engine/Plugins/Experimental/ControlRigDynamics/Source/ControlRigDynamics/Public/RigParticleSimulation.h
+++ b/Engine/Plugins/Experimental/ControlRigDynamics/Source/ControlRigDynamics/Public/RigParticleSimulation.h
@@ -96,12 +96,6 @@ struct FParticleInfo
 //======================================================================================================================
 struct FParticle
 {
-	// Call once per update
-	void Update(float InInvMass) 
-	{ 
-		InvMass = InInvMass; 
-	}
-
 	float InvMass = 0.0f; // 0 means kinematic (and track the target)
 
 	FSimVector Position = FSimVector::ZeroVector;
@@ -115,12 +109,6 @@ struct FParticle
 //======================================================================================================================
 struct FParticleTarget
 {
-	// Call once per update. InParentParticle is needed for bilateral compliance when TargetMode > 0.
-	void Update(
-		const FParticle& InParticle, const FParticle* InParentParticle,
-		float InStrength, float InDampingRatio, float InExtraDamping,
-		float InTargetVelocityInfluence, float InTargetMode);
-
 	// How we drive to the target
 	float Compliance = 0.0f; // alpha
 	float Damping = 0.0f;    // beta
@@ -132,20 +120,37 @@ struct FParticleTarget
 	// The current interpolated direction from parent to child (unit vector, per substep)
 	FSimVector TargetDirectionFromParent = FSimVector::ZeroVector;
 
+	// Angle limit. Precomputed cos/sin so the solve loop avoids trig calls.
+	// AngleLimitCompliance zero means disabled.
+	float CosAngleLimit = 1.0f;
+	float SinAngleLimit = 0.0f;
+	float AngleLimitCompliance = 0.0f;
+
 	// Accumulated
 	FSimVector Lambda = FSimVector::ZeroVector;
 };
 
 //======================================================================================================================
-struct FParticleCollider
+struct FConeLimitInfo
 {
-	// Call once per update
-	void Update(float InRadius) 
-	{ 
-		Radius = InRadius; 
-	}
+	int32 GrandparentIndex = -1;
+	int32 ParentIndex = -1;
+	int32 ChildIndex = -1;
+};
 
-	// Collision
+//======================================================================================================================
+struct FConeLimit
+{
+	float Angle = 0.0f;      // Radians
+	float Compliance = 0.0f; // alpha
+	float Damping = 0.0f;    // beta
+	// Accumulated
+	float Lambda = 0.0f;
+};
+
+//======================================================================================================================
+struct FParticleCollider
+{
 	float Radius = 0.0f;
 	bool bCollideWithColliders = true;
 };
@@ -174,9 +179,6 @@ struct FHardDistanceConstraintInfo
 //======================================================================================================================
 struct FHardDistanceConstraint
 {
-	// Call once per update
-	void Update(float InTargetDistance) { TargetDistance = InTargetDistance; }
-
 	float TargetDistance = 0.0f;
 };
 
@@ -190,11 +192,6 @@ struct FSoftDistanceConstraintInfo
 //======================================================================================================================
 struct FSoftDistanceConstraint
 {
-	// Call once per update
-	void Update(
-		const FParticle& ParentParticle, const FParticle& ChildParticle, 
-		float InTargetDistance, float InStrength, float InDampingRatio, float InExtraDamping);
-
 	float TargetDistance = 0.0f;
 
 	float Compliance = 0.0f; // alpha
@@ -208,7 +205,7 @@ struct FSoftDistanceConstraint
 // TODO investigate using the implicit shapes here (and for the others) - e.g. TBox/FImplicitBox3
 struct FBoxCollider
 {
-	// Call once per update
+	// Optionally call once per update, if the local data have changed
 	void Update(const FSimTransform& InTM, const FSimVector& InExtents)
 	{
 		TM = InTM;
@@ -222,7 +219,7 @@ struct FBoxCollider
 //======================================================================================================================
 struct FCapsuleCollider
 {
-	// Call once per update
+	// Optionally call once per update, if the local data have changed
 	void Update(const FSimTransform& InTM, float InLength, float InRadius)
 	{
 		TM = InTM;
@@ -239,7 +236,7 @@ struct FCapsuleCollider
 //======================================================================================================================
 struct FPlaneCollider
 {
-	// Call once per update
+	// Optionally call once per update, if the local data have changed
 	void Update(const FSimTransform& InTM, const FSimVector2D& InExtents)
 	{
 		TM = InTM;
@@ -254,12 +251,6 @@ struct FPlaneCollider
 //======================================================================================================================
 struct FCollider
 {
-	// Call once per update
-	void Update(const FSimTransform& InTM)
-	{
-		PreviousTM = TargetTM;
-		TargetTM = InTM;
-	}
 	FSimTransform PreviousTM;
 	FSimTransform TargetTM;
 
@@ -310,6 +301,9 @@ struct FSimulationState
 	TArray<FSoftDistanceConstraintInfo>  SoftDistanceConstraintInfos; // User-specified
 	TArray<FSoftDistanceConstraint>      SoftDistanceConstraints; 
 
+	TArray<FConeLimitInfo>               ConeLimitInfos;
+	TArray<FConeLimit>                   ConeLimits;
+
 	TArray<FCollider>                    Colliders;
 
 	// Previous substep time, used for time-corrected Verlet integration at frame boundaries
@@ -382,6 +376,8 @@ UE_API void UpdateParticle(
 // InExtraDamping            - additional damping
 // InTargetVelocityInfluence - how much target velocity contributes to damping
 // InTargetMode              - 0 = SimSpace, 1 = directional, blend in between
+// InAngleLimit              - max deviation from target direction in degrees
+// InAngleLimitStrength      - oscillation frequency for the angle limit softness
 //======================================================================================================================
 UE_API void UpdateParticleTarget(
 	FSimulationState&                  State,
@@ -392,7 +388,15 @@ UE_API void UpdateParticleTarget(
 	float                              InDampingRatio,
 	float                              InExtraDamping,
 	float                              InTargetVelocityInfluence,
-	float                              InTargetMode);
+	float                              InTargetMode,
+	float                              InAngleLimit,
+	float                              InAngleLimitStrength);
+
+//======================================================================================================================
+// Updates a single collider for the current frame. Should be called once per collider per update,
+// unless the values and target position haven't changed for two updates
+//======================================================================================================================
+UE_API void UpdateCollider(FSimulationState& State, int32 Index, const FSimTransform& TM);
 
 //======================================================================================================================
 // Updates a single skeletal (bone-chain) distance constraint's target distance from the current
@@ -418,10 +422,12 @@ UE_API void UpdateSoftDistanceConstraint(
 	float InStrength, float InDampingRatio, float InExtraDamping);
 
 //======================================================================================================================
-// Updates a single collider for the current frame. Call once per collider per update when the
-// properties have changed.
+// Updates a single cone limit for the current frame. Recomputes the angle limit and
+// compliance/damping from the given parameters. Call once per cone limit per update when values
+// have changed.
 //======================================================================================================================
-UE_API void UpdateCollider(FSimulationState& State, int32 Index, const FSimTransform& TM);
+UE_API void UpdateConeLimit(
+	FSimulationState& State, int32 Index, float InAngleDegrees, float InStrength, float InDampingRatio);
 
 //======================================================================================================================
 // Runs the simulation for one frame. Internally substeps according to MaxTimeStep/MaxNumSteps.
@@ -456,23 +462,25 @@ UE_API TArray<int32> SortParticlesRootToLeaf(FSimulationState& State);
 //======================================================================================================================
 inline void UpdateSkeletalConstraint(FSimulationState& State, const int32 Index)
 {
+	check(State.SkeletalConstraintInfos.IsValidIndex(Index));
 	const FHardDistanceConstraintInfo& Info = State.SkeletalConstraintInfos[Index];
 	check(State.ParticleInfos.IsValidIndex(Info.ParentIndex) && State.ParticleInfos.IsValidIndex(Info.ChildIndex));
 	const FSimVector& ParentPos = State.ParticleInfos[Info.ParentIndex].TargetPosition;
 	const FSimVector& ChildPos = State.ParticleInfos[Info.ChildIndex].TargetPosition;
-	State.SkeletalConstraints[Index].Update(FSimVector::Dist(ParentPos, ChildPos));
+	State.SkeletalConstraints[Index].TargetDistance = FSimVector::Dist(ParentPos, ChildPos);
 }
 
 //======================================================================================================================
 inline void UpdateHardDistanceConstraint(
 	FSimulationState& State, const int32 Index, const float InLengthMultiplier, const float InExtraLength)
 {
+	check(State.HardDistanceConstraintInfos.IsValidIndex(Index));
 	const FHardDistanceConstraintInfo& Info = State.HardDistanceConstraintInfos[Index];
 	check(State.ParticleInfos.IsValidIndex(Info.ParentIndex) && State.ParticleInfos.IsValidIndex(Info.ChildIndex));
 	const FSimVector& ParentPos = State.ParticleInfos[Info.ParentIndex].TargetPosition;
 	const FSimVector& ChildPos = State.ParticleInfos[Info.ChildIndex].TargetPosition;
-	const float Distance = FSimVector::Dist(ParentPos, ChildPos) * InLengthMultiplier + InExtraLength;
-	State.HardDistanceConstraints[Index].Update(Distance);
+	State.HardDistanceConstraints[Index].TargetDistance =
+		FSimVector::Dist(ParentPos, ChildPos) * InLengthMultiplier + InExtraLength;
 }
 
 //======================================================================================================================
@@ -480,13 +488,36 @@ inline void UpdateSoftDistanceConstraint(
 	FSimulationState& State, const int32 Index, const float InLengthMultiplier, const float InExtraLength,
 	const float InStrength, const float InDampingRatio, const float InExtraDamping)
 {
+	check(State.SoftDistanceConstraintInfos.IsValidIndex(Index));
 	const FSoftDistanceConstraintInfo& Info = State.SoftDistanceConstraintInfos[Index];
-	check(State.ParticleInfos.IsValidIndex(Info.ParentIndex) && State.ParticleInfos.IsValidIndex(Info.ChildIndex));
+	check(State.Particles.IsValidIndex(Info.ParentIndex) && State.Particles.IsValidIndex(Info.ChildIndex));
 	const FSimVector& ParentPos = State.ParticleInfos[Info.ParentIndex].TargetPosition;
 	const FSimVector& ChildPos = State.ParticleInfos[Info.ChildIndex].TargetPosition;
-	const float Distance = FSimVector::Dist(ParentPos, ChildPos) * InLengthMultiplier + InExtraLength;
-	State.SoftDistanceConstraints[Index].Update(State.Particles[Info.ParentIndex], State.Particles[Info.ChildIndex],
-		Distance, InStrength, InDampingRatio, InExtraDamping);
+	FSoftDistanceConstraint& Constraint = State.SoftDistanceConstraints[Index];
+	Constraint.TargetDistance = FSimVector::Dist(ParentPos, ChildPos) * InLengthMultiplier + InExtraLength;
+	const float W = InStrength * TWO_PI;
+	const float InvMassSum =
+		State.Particles[Info.ParentIndex].InvMass + State.Particles[Info.ChildIndex].InvMass;
+	Constraint.Compliance = InvMassSum / FMath::Max(W * W, KINDA_SMALL_NUMBER);
+	Constraint.Damping = (2.0f * InDampingRatio * W + InExtraDamping) / FMath::Max(InvMassSum, KINDA_SMALL_NUMBER);
+}
+
+//======================================================================================================================
+inline void UpdateConeLimit(
+	FSimulationState& State, const int32 Index,
+	const float InAngleDegrees, const float InStrength, const float InDampingRatio)
+{
+	check(State.ConeLimitInfos.IsValidIndex(Index));
+	const FConeLimitInfo& Info = State.ConeLimitInfos[Index];
+	check(State.Particles.IsValidIndex(Info.GrandparentIndex) && State.Particles.IsValidIndex(Info.ParentIndex)
+		&& State.Particles.IsValidIndex(Info.ChildIndex));
+	FConeLimit& ConeLimit = State.ConeLimits[Index];
+	ConeLimit.Angle = FMath::DegreesToRadians(InAngleDegrees);
+	const float W = InStrength * TWO_PI;
+	const float InvMassSum = State.Particles[Info.GrandparentIndex].InvMass
+		+ State.Particles[Info.ParentIndex].InvMass + State.Particles[Info.ChildIndex].InvMass;
+	ConeLimit.Compliance = InvMassSum / FMath::Max(W * W, KINDA_SMALL_NUMBER);
+	ConeLimit.Damping = (2.0f * InDampingRatio * W) / FMath::Max(InvMassSum, KINDA_SMALL_NUMBER);
 }
 
 //======================================================================================================================
@@ -494,7 +525,8 @@ inline void UpdateCollider(FSimulationState& State, int32 Index, const FSimTrans
 {
 	check(State.Colliders.IsValidIndex(Index));
 	RigParticleSimulation::FCollider& Collider = State.Colliders[Index];
-	Collider.Update(TM);
+	Collider.PreviousTM = Collider.TargetTM;
+	Collider.TargetTM = TM;
 }
 
 }
